Plugin Auditing - Adding an experimental plugin for auditing plugins to ensure that all the plugins the project is planning to cook don't depend on uncooked plug-in content.  Also checks sources for things like gameplay tags.

diff --git a/Engine/Plugins/Experimental/GameFeatures/Source/GameFeatures/Private/GameFeaturesProjectPolicies.cpp b/Engine/Plugins/Experimental/GameFeatures/Source/GameFeatures/Private/GameFeaturesProjectPolicies.cpp
--- a/Engine/Plugins/Experimental/GameFeatures/Source/GameFeatures/Private/GameFeaturesProjectPolicies.cpp
+++ b/Engine/Plugins/Experimental/GameFeatures/Source/GameFeatures/Private/GameFeaturesProjectPolicies.cpp
@@ -42,6 +42,11 @@ const TArray<FName> UDefaultGameFeaturesProjectPolicies::GetPreloadBundleStateFo
 	return FeatureBundles;
 }
 
+bool UGameFeaturesProjectPolicies::WillPluginBeCooked(const FString& PluginFilename, const FGameFeaturePluginDetails& PluginDetails) const
+{
+	return true;
+}
+
 void UGameFeaturesProjectPolicies::ExplicitLoadGameFeaturePlugin(const FString& PluginURL, const FGameFeaturePluginLoadComplete& CompleteDelegate, const bool bActivateGameFeatures)
 {
 	if (bActivateGameFeatures)
diff --git a/Engine/Plugins/Experimental/GameFeatures/Source/GameFeatures/Private/GameFeaturesSubsystem.cpp b/Engine/Plugins/Experimental/GameFeatures/Source/GameFeatures/Private/GameFeaturesSubsystem.cpp
--- a/Engine/Plugins/Experimental/GameFeatures/Source/GameFeatures/Private/GameFeaturesSubsystem.cpp
+++ b/Engine/Plugins/Experimental/GameFeatures/Source/GameFeatures/Private/GameFeaturesSubsystem.cpp
@@ -1145,42 +1145,26 @@ void UGameFeaturesSubsystem::LoadBuiltInGameFeaturePlugin(const TSharedRef<IPlug
 
 	UAssetManager::Get().PushBulkScanning();
 
-	const FString& PluginDescriptorFilename = Plugin->GetDescriptorFileName();
-
-	// Make sure you are in a game feature plugins folder. All GameFeaturePlugins are rooted in a GameFeatures folder.
-	if (!PluginDescriptorFilename.IsEmpty() && GetDefault<UGameFeaturesSubsystemSettings>()->IsValidGameFeaturePlugin(FPaths::ConvertRelativePathToFull(PluginDescriptorFilename)) && FPaths::FileExists(PluginDescriptorFilename))
+	FString PluginURL;
+	FGameFeaturePluginDetails PluginDetails;
+	if (GetGameFeaturePluginDetails(Plugin, PluginURL, PluginDetails))
 	{
-		FString PluginURL;
-		bool bIsFileProtocol = true;
-		if (GetPluginURLByName(Plugin->GetName(), PluginURL))
+		if (GameSpecificPolicies->IsPluginAllowed(PluginURL))
 		{
-			bIsFileProtocol = UGameFeaturesSubsystem::IsPluginURLProtocol(PluginURL, EGameFeaturePluginProtocol::File);
-		}
-		else
-		{
-			PluginURL = GetPluginURL_FileProtocol(PluginDescriptorFilename);
-		}
-
-		if (bIsFileProtocol && GameSpecificPolicies->IsPluginAllowed(PluginURL))
-		{
-			FGameFeaturePluginDetails PluginDetails;
-			if (GetGameFeaturePluginDetails(PluginDescriptorFilename, PluginDetails))
+			FBuiltInGameFeaturePluginBehaviorOptions BehaviorOptions;
+			const bool bShouldProcess = AdditionalFilter(Plugin->GetDescriptorFileName(), PluginDetails, BehaviorOptions);
+			if (bShouldProcess)
 			{
-				FBuiltInGameFeaturePluginBehaviorOptions BehaviorOptions;
-				bool bShouldProcess = AdditionalFilter(PluginDescriptorFilename, PluginDetails, BehaviorOptions);
-				if (bShouldProcess)
-				{
-					UGameFeaturePluginStateMachine* StateMachine = FindOrCreateGameFeaturePluginStateMachine(PluginURL);
+				UGameFeaturePluginStateMachine* StateMachine = FindOrCreateGameFeaturePluginStateMachine(PluginURL);
 
-					const EBuiltInAutoState InitialAutoState = (BehaviorOptions.AutoStateOverride != EBuiltInAutoState::Invalid) ? BehaviorOptions.AutoStateOverride : PluginDetails.BuiltInAutoState;
-						
-					const EGameFeaturePluginState DestinationState = ConvertInitialFeatureStateToTargetState(InitialAutoState);
+				const EBuiltInAutoState InitialAutoState = (BehaviorOptions.AutoStateOverride != EBuiltInAutoState::Invalid) ? BehaviorOptions.AutoStateOverride : PluginDetails.BuiltInAutoState;
+				
+				const EGameFeaturePluginState DestinationState = ConvertInitialFeatureStateToTargetState(InitialAutoState);
 
-					// If we're already at the destination or beyond, don't transition back
-					FGameFeaturePluginStateRange Destination(DestinationState, EGameFeaturePluginState::Active);
-					ChangeGameFeatureDestination(StateMachine, Destination, 
-						FGameFeaturePluginChangeStateComplete::CreateUObject(this, &ThisClass::LoadBuiltInGameFeaturePluginComplete, StateMachine, Destination));
-				}
+				// If we're already at the destination or beyond, don't transition back
+				FGameFeaturePluginStateRange Destination(DestinationState, EGameFeaturePluginState::Active);
+				ChangeGameFeatureDestination(StateMachine, Destination, 
+					FGameFeaturePluginChangeStateComplete::CreateUObject(this, &ThisClass::LoadBuiltInGameFeaturePluginComplete, StateMachine, Destination));
 			}
 		}
 	}
@@ -1352,6 +1336,32 @@ EGameFeaturePluginState UGameFeaturesSubsystem::GetPluginState(FGameFeaturePlugi
 	}
 }
 
+bool UGameFeaturesSubsystem::GetGameFeaturePluginDetails(const TSharedRef<IPlugin>& Plugin, FString& OutPluginURL, struct FGameFeaturePluginDetails& OutPluginDetails) const
+{
+	const FString& PluginDescriptorFilename = Plugin->GetDescriptorFileName();
+
+	// Make sure you are in a game feature plugins folder. All GameFeaturePlugins are rooted in a GameFeatures folder.
+	if (!PluginDescriptorFilename.IsEmpty() && GetDefault<UGameFeaturesSubsystemSettings>()->IsValidGameFeaturePlugin(FPaths::ConvertRelativePathToFull(PluginDescriptorFilename)) && FPaths::FileExists(PluginDescriptorFilename))
+	{
+		bool bIsFileProtocol = true;
+		if (GetPluginURLByName(Plugin->GetName(), OutPluginURL))
+		{
+			bIsFileProtocol = UGameFeaturesSubsystem::IsPluginURLProtocol(OutPluginURL, EGameFeaturePluginProtocol::File);
+		}
+		else
+		{
+			OutPluginURL = GetPluginURL_FileProtocol(PluginDescriptorFilename);
+		}
+
+		if (bIsFileProtocol)
+		{
+			return GetGameFeaturePluginDetails(PluginDescriptorFilename, OutPluginDetails);
+		}
+	}
+
+	return false;
+}
+
 bool UGameFeaturesSubsystem::GetGameFeaturePluginDetails(const FString& PluginDescriptorFilename, FGameFeaturePluginDetails& OutPluginDetails) const
 {
 	TRACE_CPUPROFILER_EVENT_SCOPE(GFP_GetPluginDetails);
diff --git a/Engine/Plugins/Experimental/GameFeatures/Source/GameFeatures/Public/GameFeaturesProjectPolicies.h b/Engine/Plugins/Experimental/GameFeatures/Source/GameFeatures/Public/GameFeaturesProjectPolicies.h
--- a/Engine/Plugins/Experimental/GameFeatures/Source/GameFeatures/Public/GameFeaturesProjectPolicies.h
+++ b/Engine/Plugins/Experimental/GameFeatures/Source/GameFeatures/Public/GameFeaturesProjectPolicies.h
@@ -23,6 +23,9 @@ public:
 	// Called when the game feature manager is shut down
 	virtual void ShutdownGameFeatureManager() { }
 
+	// Called to determined the expected state of a plugin under the WhenLoading conditions.
+	virtual bool WillPluginBeCooked(const FString& PluginFilename, const FGameFeaturePluginDetails& PluginDetails) const;
+
 	// Called when a game feature plugin enters the Loading state to determine additional assets to load
 	virtual TArray<FPrimaryAssetId> GetPreloadAssetListForGameFeature(const UGameFeatureData* GameFeatureToLoad, bool bIncludeLoadedAssets = false) const { return TArray<FPrimaryAssetId>(); }
 
diff --git a/Engine/Plugins/Experimental/GameFeatures/Source/GameFeatures/Public/GameFeaturesSubsystem.h b/Engine/Plugins/Experimental/GameFeatures/Source/GameFeatures/Public/GameFeaturesSubsystem.h
--- a/Engine/Plugins/Experimental/GameFeatures/Source/GameFeatures/Public/GameFeaturesSubsystem.h
+++ b/Engine/Plugins/Experimental/GameFeatures/Source/GameFeatures/Public/GameFeaturesSubsystem.h
@@ -496,6 +496,8 @@ public:
 	/** Returns the current state of the state machine for the specified plugin PluginIdentifier */
 	EGameFeaturePluginState GetPluginState(FGameFeaturePluginIdentifier PluginIdentifier) const;
 
+	/** Gets relevant properties out of a uplugin file */
+	bool GetGameFeaturePluginDetails(const TSharedRef<IPlugin>& Plugin, FString& OutPluginURL, struct FGameFeaturePluginDetails& OutPluginDetails) const;
 
 	/** Determine the initial feature state for a built-in plugin */
 	static EBuiltInAutoState DetermineBuiltInInitialFeatureState(TSharedPtr<FJsonObject> Descriptor, const FString& ErrorContext);
diff --git a/Engine/Plugins/Experimental/PluginAudit/PluginAudit.uplugin b/Engine/Plugins/Experimental/PluginAudit/PluginAudit.uplugin
new file mode 100644
--- /dev/null
+++ b/Engine/Plugins/Experimental/PluginAudit/PluginAudit.uplugin
@@ -0,0 +1,36 @@
+{
+	"FileVersion" : 3,
+	"Version" : 1,
+	"VersionName" : "1.0",
+	"FriendlyName" : "Plugin Audit",
+	"Description" : "Editor plugin for auditing plugin connectivity.",
+	"Category" : "Editor",
+	"CreatedBy" : "Epic Games, Inc.",
+	"CreatedByURL" : "http://epicgames.com",
+	"DocsURL" : "",
+	"MarketplaceURL" : "",
+	"SupportURL" : "",
+	"EnabledByDefault" : false,
+	"CanContainContent" : false,
+	"IsBetaVersion" : true,
+	"Installed" : false,
+	
+	"Modules" :
+	[
+		{
+			"Name" : "PluginAudit",
+			"Type": "UncookedOnly" ,
+			"LoadingPhase" : "PreDefault"
+		}
+	],
+	"Plugins": [
+		{
+			"Name": "GameFeatures",
+			"Enabled": true
+		},
+		{
+			"Name": "AssetManagerEditor",
+			"Enabled": true
+		}
+	]
+}
\ No newline at end of file
diff --git a/Engine/Plugins/Experimental/PluginAudit/Source/PluginAudit.Build.cs b/Engine/Plugins/Experimental/PluginAudit/Source/PluginAudit.Build.cs
new file mode 100644
--- /dev/null
+++ b/Engine/Plugins/Experimental/PluginAudit/Source/PluginAudit.Build.cs
@@ -0,0 +1,61 @@
+// Copyright Epic Games, Inc. All Rights Reserved.
+
+using UnrealBuildTool;
+
+public class PluginAudit : ModuleRules
+{
+	public PluginAudit(ReadOnlyTargetRules Target) : base(Target)
+	{
+		PrivateIncludePathModuleNames.AddRange(
+			new string[] {
+				"StatsViewer",
+			}
+		);
+
+		PublicDependencyModuleNames.AddRange(
+			new string[] { 
+				"Core",
+				"CoreUObject",
+				"Engine",
+				"TargetPlatform",
+			}
+		);
+		
+		PrivateDependencyModuleNames.AddRange(
+			new string[] {
+				"Slate",
+				"SlateCore",
+                "ApplicationCore",
+                "InputCore",
+				"EditorFramework",
+				"UnrealEd",
+				"AssetRegistry",
+				"Json",
+				"CollectionManager",
+				"ContentBrowser",
+				"ContentBrowserData",
+				"WorkspaceMenuStructure",
+				"AssetDefinition",
+				"AssetTools",
+				"PropertyEditor",
+				"GraphEditor",
+				"BlueprintGraph",
+				"KismetCompiler",
+				"MessageLog",
+				"LevelEditor",
+				"SandboxFile",
+				"EditorWidgets",
+				"TreeMap",
+				"ToolMenus",
+				"ToolWidgets",
+				"SourceControl",
+				"SourceControlWindows",
+				"UncontrolledChangelists",
+				"Projects",
+				"GameFeatures",
+				"GameplayTags",
+				"AssetManagerEditor",
+			}
+		);
+	}
+}
diff --git a/Engine/Plugins/Experimental/PluginAudit/Source/Private/PluginAuditModule.cpp b/Engine/Plugins/Experimental/PluginAudit/Source/Private/PluginAuditModule.cpp
new file mode 100644
--- /dev/null
+++ b/Engine/Plugins/Experimental/PluginAudit/Source/Private/PluginAuditModule.cpp
@@ -0,0 +1,56 @@
+// Copyright Epic Games, Inc. All Rights Reserved.
+
+#include "Modules/ModuleInterface.h"
+#include "SPluginAuditBrowser.h"
+#include "Widgets/Docking/SDockTab.h"
+#include "WorkspaceMenuStructure.h"
+#include "WorkspaceMenuStructureModule.h"
+
+#define LOCTEXT_NAMESPACE "PluginAudit"
+
+
+class FPluginAuditModule : public IModuleInterface
+{
+public:
+	// IModuleInterface interface
+	virtual void StartupModule() override;
+	virtual void ShutdownModule() override;
+	// End of IModuleInterface interface
+
+private:
+	static const FName PluginAuditTabName;
+	TSharedRef<SDockTab> SpawnPluginAuditTab(const FSpawnTabArgs& Args);
+};
+
+const FName FPluginAuditModule::PluginAuditTabName = TEXT("PluginAudit");
+
+///////////////////////////////////////////
+
+IMPLEMENT_MODULE(FPluginAuditModule, PluginAudit);
+
+
+void FPluginAuditModule::StartupModule()
+{
+	FGlobalTabmanager::Get()->RegisterNomadTabSpawner(PluginAuditTabName, FOnSpawnTab::CreateRaw(this, &FPluginAuditModule::SpawnPluginAuditTab))
+		.SetDisplayName(LOCTEXT("PluginAuditTitle", "Plugin Audit"))
+		.SetTooltipText(LOCTEXT("PluginAuditTooltip", "Open Plugin Audit window, allows viewing detailed information about plugin references."))
+		.SetGroup(WorkspaceMenu::GetMenuStructure().GetDeveloperToolsAuditCategory())
+		.SetIcon(FSlateIcon(FAppStyle::GetAppStyleSetName(), "LevelEditor.Audit"));
+	FGlobalTabmanager::Get()->RegisterDefaultTabWindowSize(PluginAuditTabName, FVector2D(1080, 600));
+}
+
+void FPluginAuditModule::ShutdownModule()
+{
+	FGlobalTabmanager::Get()->UnregisterNomadTabSpawner(PluginAuditTabName);
+}
+
+TSharedRef<SDockTab> FPluginAuditModule::SpawnPluginAuditTab(const FSpawnTabArgs& Args)
+{
+	return SNew(SDockTab)
+		.TabRole(ETabRole::NomadTab)
+		[
+			SNew(SPluginAuditBrowser)
+		];
+}
+
+#undef LOCTEXT_NAMESPACE
diff --git a/Engine/Plugins/Experimental/PluginAudit/Source/Private/SPluginAuditBrowser.cpp b/Engine/Plugins/Experimental/PluginAudit/Source/Private/SPluginAuditBrowser.cpp
new file mode 100644
--- /dev/null
+++ b/Engine/Plugins/Experimental/PluginAudit/Source/Private/SPluginAuditBrowser.cpp
@@ -0,0 +1,601 @@
+// Copyright Epic Games, Inc. All Rights Reserved.
+
+#include "SPluginAuditBrowser.h"
+
+#include "GameFeaturesProjectPolicies.h"
+#include "GameFeaturesSubsystem.h"
+#include "GameplayTagsManager.h"
+#include "GameplayTagsSettings.h"
+#include "AssetRegistry/AssetRegistryModule.h"
+#include "AssetRegistry/AssetData.h"
+#include "Interfaces/IPluginManager.h"
+#include "Algo/Copy.h"
+#include "Misc/Paths.h"
+#include "Async/ParallelFor.h"
+#include "Misc/ScopedSlowTask.h"
+#include "IMessageLogListing.h"
+#include "MessageLogModule.h"
+#include "Misc/UObjectToken.h"
+#include "AssetManagerEditorModule.h"
+
+#define LOCTEXT_NAMESPACE "SPluginAuditBrowser"
+
+namespace PluginAudit
+{
+	FName PluginAuditLogName = TEXT("Plugin Audit");
+}
+
+void SPluginAuditBrowser::Construct(const FArguments& InArgs)
+{
+	CreateLogListing();
+	BuildPluginList();
+
+	FMessageLogModule& MessageLogModule = FModuleManager::LoadModuleChecked<FMessageLogModule>("MessageLog");
+
+	RefreshToolBar();
+
+	ChildSlot
+	[
+		SNew(SVerticalBox)
+		+ SVerticalBox::Slot()
+		.AutoHeight()
+		[
+			SNew(SBorder)
+			.BorderImage(FAppStyle::GetBrush("Brushes.Panel"))
+			[
+				UToolMenus::Get()->GenerateWidget("PluginAudit.MainToolBar", FToolMenuContext())
+			]
+		]
+
+		+ SVerticalBox::Slot()
+		.FillHeight(1.0f)
+		[
+			SNew(SSplitter)
+			.Orientation(Orient_Horizontal)
+
+			+ SSplitter::Slot()
+			.Value(0.30f)
+			[
+				SNew(SListView< TSharedRef<FCookedPlugin> >)
+				.ListItemsSource(&CookedPlugins)
+				.OnGenerateRow(this, &SPluginAuditBrowser::MakeCookedPluginRow)
+			]
+			+ SSplitter::Slot()
+			.Value(0.70f)
+			[
+				MessageLogModule.CreateLogListingWidget(LogListing.ToSharedRef())
+			]
+		]
+	];
+}
+
+void SPluginAuditBrowser::BuildPluginList()
+{
+	const UGameFeaturesSubsystem& GameFeaturesSubsystem = UGameFeaturesSubsystem::Get();
+	const UGameFeaturesProjectPolicies& Policy = UGameFeaturesSubsystem::Get().GetPolicy();
+
+	TArray<TSharedRef<IPlugin>> AllPlugins = IPluginManager::Get().GetDiscoveredPlugins();
+
+	IncludedGameFeaturePlugins.Reset();
+	ExcludedGameFeaturePlugins.Reset();
+	for (const TSharedRef<IPlugin>& Plugin : AllPlugins)
+	{
+		FString PluginURL;
+		FGameFeaturePluginDetails PluginDetails;
+		if (GameFeaturesSubsystem.GetGameFeaturePluginDetails(Plugin, PluginURL, PluginDetails))
+		{
+			if (Policy.WillPluginBeCooked(Plugin->GetDescriptorFileName(), PluginDetails))
+			{
+				// We're actually including this plugin.
+				IncludedGameFeaturePlugins.Add(Plugin);
+				continue;
+			}
+
+			ExcludedGameFeaturePlugins.Add(Plugin);
+		}
+	}
+
+	for (const TSharedRef<IPlugin>& Plugin : IncludedGameFeaturePlugins)
+	{
+		CookedPlugins.Add(MakeShared<FCookedPlugin>(Plugin));
+	}
+
+	CookedPlugins.Sort([](const TSharedRef<FCookedPlugin>& A, const TSharedRef<FCookedPlugin>& B)
+	{
+		return A->Plugin->GetName() < B->Plugin->GetName();
+	});
+}
+
+void SPluginAuditBrowser::RefreshToolBar()
+{
+	UToolMenus* ToolMenus = UToolMenus::Get();
+	UToolMenu* MainToolBar = UToolMenus::Get()->RegisterMenu("PluginAudit.MainToolBar", NAME_None, EMultiBoxType::SlimHorizontalToolBar);
+	MainToolBar->StyleName = "AssetEditorToolbar";
+	{
+		FToolMenuSection& PlaySection = MainToolBar->AddSection("Actions");
+
+		FToolMenuEntry SaveEntry = FToolMenuEntry::InitToolBarButton(
+			"Refresh",
+			FUIAction(FExecuteAction::CreateSP(this, &SPluginAuditBrowser::RefreshViolations)),
+			LOCTEXT("SaveDirtyPackages", "Refresh"),
+			LOCTEXT("SaveDirtyPackagesTooltip", "Refreshes the audit results based on the enabled plugins."),
+			FSlateIcon(FAppStyle::Get().GetStyleSetName(), "Icons.Refresh")
+		);
+
+		PlaySection.AddEntry(SaveEntry);
+	}
+}
+
+TSharedRef<ITableRow> SPluginAuditBrowser::MakeCookedPluginRow(TSharedRef<SPluginAuditBrowser::FCookedPlugin> InItem, const TSharedRef<STableViewBase>& OwnerTable)
+{
+	return SNew(STableRow< TSharedRef<FCookedPlugin> >, OwnerTable)
+	[
+		SNew(SHorizontalBox)
+		+ SHorizontalBox::Slot()
+		.AutoWidth()
+		.VAlign(VAlign_Center)
+		[
+			SNew(SCheckBox)
+			.IsChecked_Lambda([InItem]() -> ECheckBoxState
+			{
+				return InItem->bSimulateDisabled ? ECheckBoxState::Unchecked : ECheckBoxState::Checked;
+			})
+			.OnCheckStateChanged_Lambda([InItem](ECheckBoxState State)
+			{
+				InItem->bSimulateDisabled = (State != ECheckBoxState::Checked);
+			})
+		]
+
+		+ SHorizontalBox::Slot()
+		.FillWidth(1.0f)
+		.VAlign(VAlign_Center)
+		[
+			SNew(STextBlock)
+			.Text(FText::FromString(InItem->Plugin->GetFriendlyName()))
+		]
+	];
+}
+
+void SPluginAuditBrowser::CreateLogListing()
+{
+	FMessageLogModule& MessageLogModule = FModuleManager::LoadModuleChecked<FMessageLogModule>("MessageLog");
+
+	FMessageLogInitializationOptions LogOptions;
+	LogOptions.bShowInLogWindow = false;
+	LogOptions.bAllowClear = false;
+	LogOptions.bShowPages = false;
+	LogOptions.bShowFilters = false;
+	LogOptions.MaxPageCount = 1;
+
+	LogListing = MessageLogModule.CreateLogListing(PluginAudit::PluginAuditLogName, LogOptions);
+}
+
+void SPluginAuditBrowser::RefreshViolations()
+{
+	TArray<TSharedRef<IPlugin>> IncludedPlugins = IncludedGameFeaturePlugins;
+	TArray<TSharedRef<IPlugin>> ExcludedPlugins = ExcludedGameFeaturePlugins;
+	for (const TSharedPtr<FCookedPlugin>& CookedPlugin : CookedPlugins)
+	{
+		if (CookedPlugin->bSimulateDisabled)
+		{
+			IncludedPlugins.Remove(CookedPlugin->Plugin);
+			ExcludedPlugins.Add(CookedPlugin->Plugin);
+		}
+	}
+
+	TArray<TSharedRef<FTokenizedMessage>> Violations = ScanForViolations(IncludedPlugins, ExcludedPlugins);
+
+	LogListing->ClearMessages();
+
+	for (const TSharedRef<FTokenizedMessage>& Violation : Violations)
+	{
+		LogListing->AddMessage(Violation);
+	}
+}
+
+/*static*/ TArray<TSharedRef<FTokenizedMessage>> SPluginAuditBrowser::ScanForViolations(TArray<TSharedRef<IPlugin>> InIncludedGameFeaturePlugins, TArray<TSharedRef<IPlugin>> InExcludedGameFeaturePlugins)
+{
+	IAssetRegistry* AssetRegistry = IAssetRegistry::Get();
+	AssetRegistry->WaitForCompletion();
+	
+    UGameplayTagsManager& Manager = UGameplayTagsManager::Get();
+
+	TArray<FGameFeaturePlugin> ExcludedPlugins;
+	GetGameFeaturePlugins(InExcludedGameFeaturePlugins, ExcludedPlugins);
+
+	TArray<FGameFeaturePlugin> IncludedPlugins;
+	GetGameFeaturePlugins(InIncludedGameFeaturePlugins, IncludedPlugins);
+
+	TArray<TSharedRef<FTokenizedMessage>> Violations;
+	FCriticalSection CS;
+
+	FScopedSlowTask SlowTask(IncludedPlugins.Num() + ExcludedPlugins.Num(), LOCTEXT("Examining Plugins", "Examining Plugins..."));
+	SlowTask.MakeDialog();
+
+	const FName GameplayTagStructPackage = FGameplayTag::StaticStruct()->GetOutermost()->GetFName();
+	const FName NAME_GameplayTag = FGameplayTag::StaticStruct()->GetFName();
+	
+	// Included plugins.
+	for (const FGameFeaturePlugin& Plugin : IncludedPlugins)
+	{
+		SlowTask.EnterProgressFrame(1.f, FText::FromString(Plugin.Plugin->GetFriendlyName()));
+
+		FARFilter Filter;
+		Filter.bRecursivePaths = true;
+		Filter.PackagePaths.Add(Plugin.ContentRoot);
+
+		// Query for a list of assets in the selected paths
+		TArray<FAssetData> AssetsInPlugin;
+		AssetRegistry->GetAssets(Filter, AssetsInPlugin);
+
+		ParallelFor(AssetsInPlugin.Num(), [&](int32 AssetIndex)
+		{
+			const FAssetData& AssetInPlugin = AssetsInPlugin[AssetIndex];
+			const FAssetIdentifier AssetId = FAssetIdentifier(AssetInPlugin.PackageName);
+
+			//TODO How can i tell if Asset is an editor only asset and can thus do whatever?  Since it wont matter if it
+			// references a bad tag, in the grand scheme.
+			const UE::AssetRegistry::FDependencyQuery DependencyRequirements(UE::AssetRegistry::EDependencyQuery::Game);
+
+			TArray<FAssetIdentifier> FoundDependencies;
+			if (AssetRegistry->GetDependencies(FAssetIdentifier(AssetInPlugin.PackageName), FoundDependencies, UE::AssetRegistry::EDependencyCategory::All, DependencyRequirements))
+			{
+				for (const FAssetIdentifier& DependencyId : FoundDependencies)
+				{
+					// Ensure for all included plugins that they correctly depend on any tag they reference.
+					if (DependencyId.ObjectName == NAME_GameplayTag && DependencyId.PackageName == GameplayTagStructPackage)
+					{
+						TArray<TSharedPtr<IPlugin>> PossibleSources;
+						const EDoesPluginDependGameplayTagSource DependencyResult = DoesPluginDependOnGameplayTagSource(Manager, Plugin.Plugin, DependencyId.ValueName, PossibleSources);
+						if (DependencyResult != EDoesPluginDependGameplayTagSource::Yes)
+						{
+							FString PluginPath = Plugin.ContentRoot.ToString();
+							TSharedPtr<IPlugin> ReferencerPlugin = Plugin.Plugin;
+							FAssetIdentifier Referencer = AssetId;
+							FAssetIdentifier Asset = DependencyId;
+							//Violation.AssetPossibleSources = PossibleSources;
+
+							TSharedRef<FTokenizedMessage> Message = FTokenizedMessage::Create(EMessageSeverity::Error);
+							Message->AddToken(FTextToken::Create(LOCTEXT("GameplayTagReference", "Gameplay Tag Reference")));
+							Message->AddToken(FTextToken::Create(FText::FromString(TEXT(":"))));
+							Message->AddToken(FTextToken::Create(FText::FormatNamed(LOCTEXT("DependingGameplayTagUnknown", "The {ReferencerPlugin} contains "), 
+								TEXT("ReferencerPlugin"), FText::FromString(ReferencerPlugin->GetName())
+							)));
+							Message->AddToken(
+								FAssetNameToken::Create(Referencer.ToString())->OnMessageTokenActivated(FOnMessageTokenActivated::CreateLambda([Referencer](const TSharedRef<class IMessageToken>&) {
+									IAssetManagerEditorModule::Get().OpenReferenceViewerUI({ Referencer });
+								}))
+							);
+							Message->AddToken(FTextToken::Create(LOCTEXT("AndItDependsOnGameplayTag", "and it depends on the GameplayTag ")));
+							Message->AddToken(
+								FAssetNameToken::Create(Asset.ToString())->OnMessageTokenActivated(FOnMessageTokenActivated::CreateLambda([Asset](const TSharedRef<class IMessageToken>&) {
+									IAssetManagerEditorModule::Get().OpenReferenceViewerUI({ Asset });
+								}))
+							);
+							if (DependencyResult == EDoesPluginDependGameplayTagSource::UnknownTag)
+							{
+								Message->AddToken(FTextToken::Create(LOCTEXT("FromAndUnknownPlugin", " from a plug-in.  The gameplay tag's source is unknown so it's probably in a plugin that's not registered as a dependency of this plug-in, and nothing else loads the plug-in so we can't figure out where it's suppose to come from.")));
+							}
+							else
+							{
+								Message->AddToken(FTextToken::Create(
+									FText::FormatNamed(LOCTEXT("", " from {AssetPlugin}. The {ReferencerPlugin} needs to depend on {AssetPlugin} in its .uplugin file."),
+										TEXT("ReferencerPlugin"), FText::FromString(ReferencerPlugin->GetName()),
+										TEXT("AssetPlugin"), FText::FromString(FPackageName::GetPackageMountPoint(Asset.PackageName.ToString()).ToString())
+									)
+								));
+							}
+
+							{
+								FScopeLock Lock(&CS);
+								Violations.Add(MoveTemp(Message));
+							}
+						}
+					}
+				}
+			}
+		});
+	}
+
+	// Excluded plugins
+	for (const FGameFeaturePlugin& Plugin : ExcludedPlugins)
+	{
+		SlowTask.EnterProgressFrame(1.f, FText::FromString(Plugin.Plugin->GetFriendlyName()));
+		
+		FARFilter Filter;
+		Filter.bRecursivePaths = true;
+		Filter.PackagePaths.Add(Plugin.ContentRoot);
+
+		// Query for a list of assets in the selected paths
+		TArray<FAssetData> AssetsInPlugin;
+		AssetRegistry->GetAssets(Filter, AssetsInPlugin);
+		
+		TArray<FAssetIdentifier> AssetIdsInPlugin;
+		Algo::Transform(AssetsInPlugin, AssetIdsInPlugin, [](const FAssetData& Asset){ return FAssetIdentifier(Asset.PackageName); });
+
+		// Add all the script packages for the plugin so we can find any references to native classes, structs or enums.
+		for (const FString& ScriptPackageName : Plugin.ScriptPackages)
+		{
+			AssetIdsInPlugin.Add(FAssetIdentifier(*ScriptPackageName));
+		}
+
+		// Determine any gameplay tags that are unique to this plugin that might be referenced.
+		{
+			TArray<FGameplayTag> ContentTags;
+			Manager.FindTagsWithSource(Plugin.ContentRoot.ToString(), ContentTags);
+
+			for (const FGameplayTag& Tag : ContentTags)
+			{
+				if (IsTagOnlyAvailableFromExcludedSources(Manager, Tag, ExcludedPlugins))
+				{
+					FAssetIdentifier TagId = FAssetIdentifier(FGameplayTag::StaticStruct(), Tag.GetTagName());
+					AssetIdsInPlugin.Add(TagId);
+				}
+			}
+		}
+
+		for (const FString& ModuleName : Plugin.ModuleNames)
+		{
+			TArray<FGameplayTag> NativeTags;
+			Manager.FindTagsWithSource(ModuleName, NativeTags);
+
+			for (const FGameplayTag& Tag : NativeTags)
+			{
+				if (IsTagOnlyAvailableFromExcludedSources(Manager, Tag, ExcludedPlugins))
+				{
+					FAssetIdentifier TagId = FAssetIdentifier(FGameplayTag::StaticStruct(), Tag.GetTagName());
+					AssetIdsInPlugin.Add(TagId);
+				}
+			}
+		}
+		
+		
+		// Burn through the assets in the disabled plugin including any native code packages, we need to figure out
+		// from the asset registry who - if anyone references these things.
+		ParallelFor(AssetIdsInPlugin.Num(), [&](int32 AssetIndex)
+		{
+			const FAssetIdentifier& AssetId = AssetIdsInPlugin[AssetIndex];
+
+			TArray<FAssetDependency> Referencers;
+			AssetRegistry->GetReferencers(AssetId, Referencers);
+
+			for (const FAssetDependency& Reference : Referencers)
+			{
+				// If this reference is editor only, we can ignore it.  We only care about game referencing disabled stuff.
+				if (!EnumHasAnyFlags(Reference.Properties, UE::AssetRegistry::EDependencyProperty::Game))
+				{
+					continue;
+				}
+
+				const FName PackageMountPoint = FPackageName::GetPackageMountPoint(Reference.AssetId.PackageName.ToString(), false);
+				if (IncludedPlugins.ContainsByPredicate([&PackageMountPoint](const FGameFeaturePlugin& Plugin){ return Plugin.ContentRoot == PackageMountPoint; } ))
+				{				
+					FString PluginPath = Plugin.ContentRoot.ToString();
+					FAssetIdentifier Asset = AssetId;
+					TSharedPtr<IPlugin> AssetPlugin = Plugin.Plugin;
+					FAssetIdentifier Referencer = Reference.AssetId;
+					TSharedPtr<IPlugin> ReferencerPlugin = IPluginManager::Get().FindPluginFromPath(Reference.AssetId.PackageName.ToString());
+
+					const FText ReasonFormat = LOCTEXT("IncludedPluginDependsOnExcludedPluginContent", "The {ReferencerPlugin} contains {Referencer} and it depends on {Asset} from {AssetPlugin}. The {AssetPlugin} is disabled or sunset and so can not be referenced.");
+	
+					TSharedRef<FTokenizedMessage> Message = FTokenizedMessage::Create(EMessageSeverity::Error);
+					Message->AddToken(FTextToken::Create(LOCTEXT("IncludedExcludedContent", "Disabled Content Reference")));
+					Message->AddToken(FTextToken::Create(FText::FromString(TEXT(":"))));
+					Message->AddToken(FTextToken::Create(
+						FText::FormatNamed(LOCTEXT("TheReferencePlugin","The {ReferencerPlugin} contains "),
+							TEXT("ReferencerPlugin"), FText::FromString(ReferencerPlugin->GetName()),
+							TEXT("AssetPlugin"), FText::FromString(FPackageName::GetPackageMountPoint(Asset.PackageName.ToString()).ToString())
+						)
+					));
+					Message->AddToken(
+						FAssetNameToken::Create(Referencer.ToString())->OnMessageTokenActivated(FOnMessageTokenActivated::CreateLambda([Referencer](const TSharedRef<class IMessageToken>&) {
+							IAssetManagerEditorModule::Get().OpenReferenceViewerUI({ Referencer });
+						}))
+					);
+					Message->AddToken(FTextToken::Create(LOCTEXT("AndItDependsOnAsset", " and it depends on ")));
+					Message->AddToken(
+						FAssetNameToken::Create(Asset.ToString())->OnMessageTokenActivated(FOnMessageTokenActivated::CreateLambda([Asset](const TSharedRef<class IMessageToken>&) {
+							IAssetManagerEditorModule::Get().OpenReferenceViewerUI({ Asset });
+						}))
+					);
+					Message->AddToken(FTextToken::Create(
+						FText::FormatNamed(LOCTEXT("FromAssetPlugin", " from {AssetPlugin}. The {AssetPlugin} is disabled or sunset and so can not be referenced."),
+							TEXT("AssetPlugin"), FText::FromString(FPackageName::GetPackageMountPoint(Asset.PackageName.ToString()).ToString())
+						)
+					));
+
+					{
+						FScopeLock Lock(&CS);
+						Violations.Add(MoveTemp(Message));
+					}
+				}
+			}
+		});
+	}
+
+	return Violations;
+}
+
+/*static*/ TArray<TSharedPtr<IPlugin>> SPluginAuditBrowser::GetTagSourcePlugins(const UGameplayTagsManager& Manager, FName TagName)
+{
+	TArray<TSharedRef<IPlugin>> AllPlugins = IPluginManager::Get().GetDiscoveredPlugins();
+	
+	TArray<TSharedPtr<IPlugin>> SourcePlugins;
+	if (const FGameplayTagSource* TagSource = Manager.FindTagSource(TagName))
+	{
+		FName TagPackageName;
+		switch (TagSource->SourceType)
+		{
+		case EGameplayTagSourceType::TagList:
+			if (TagSource->SourceTagList)
+			{
+				const FString ContentFilePath = FPaths::GetPath(TagSource->SourceTagList->ConfigFileName) / TEXT("../../Content/");
+				FString RootContentPath;
+				if (FPackageName::TryConvertFilenameToLongPackageName(ContentFilePath, RootContentPath))
+				{
+					TagPackageName = FName(*RootContentPath);
+				}
+			}
+			break;
+		case EGameplayTagSourceType::DataTable:
+			TagPackageName = TagSource->SourceName;
+			break;
+		case EGameplayTagSourceType::Native:
+			TagPackageName = TagSource->SourceName;
+		default:
+			break;
+		}
+
+		for (const TSharedRef<IPlugin>& Plugin : AllPlugins)
+		{
+			FString ContentRoot = FString::Printf(TEXT("/%s/"), *FPaths::GetBaseFilename(Plugin->GetDescriptorFileName()));
+			if (TagPackageName.ToString().StartsWith(ContentRoot))
+			{
+				SourcePlugins.Add(Plugin);
+				continue;
+			}
+			
+			for (const FModuleDescriptor& Module : Plugin->GetDescriptor().Modules)
+			{
+				if (TagPackageName == Module.Name)
+				{
+					SourcePlugins.Add(Plugin);
+					break;
+				}
+
+				const FName ModuleScriptPackage = FPackageName::GetModuleScriptPackageName(Module.Name);
+				if (TagPackageName == ModuleScriptPackage)
+				{
+					SourcePlugins.Add(Plugin);
+					break;
+				}
+			}
+		}
+	}
+
+	return SourcePlugins;
+}
+
+/*static*/ SPluginAuditBrowser::EDoesPluginDependGameplayTagSource SPluginAuditBrowser::DoesPluginDependOnGameplayTagSource(const UGameplayTagsManager& Manager, const TSharedPtr<IPlugin>& DependentPlugin, FName TagName, TArray<TSharedPtr<IPlugin>>& OutPossibleSources)
+{
+	FString Comment;
+	TArray<FName> TagSources;
+	bool bIsTagExplicit, bIsRestrictedTag, bAllowNonRestrictedChildren;
+	if (Manager.GetTagEditorData(TagName, Comment, TagSources, bIsTagExplicit, bIsRestrictedTag, bAllowNonRestrictedChildren))
+	{
+		const TArray<FModuleDescriptor>& Modules = DependentPlugin->GetDescriptor().Modules;
+		const TArray<FPluginReferenceDescriptor>& Plugins = DependentPlugin->GetDescriptor().Plugins;
+
+		TArray<TSharedPtr<IPlugin>> TagSourcePlugins = GetTagSourcePlugins(Manager, TagName);
+		if (TagSourcePlugins.Num() == 0)
+		{
+			// Must be a builtin module?
+			return EDoesPluginDependGameplayTagSource::Yes;
+		}
+
+		for (const TSharedPtr<IPlugin>& TagSourcePlugin : TagSourcePlugins)
+		{
+			if (TagSourcePlugin == DependentPlugin)
+			{
+				// The dependent plugin is the source of the tag, so we good.
+				return EDoesPluginDependGameplayTagSource::Yes;
+			}
+			else if (Plugins.ContainsByPredicate([&TagSourcePlugin](const FPluginReferenceDescriptor& PluginDescriptor){ return PluginDescriptor.Name == TagSourcePlugin->GetName(); }))
+			{
+				return EDoesPluginDependGameplayTagSource::Yes;
+			}
+		}
+
+		OutPossibleSources = MoveTemp(TagSourcePlugins);
+		return EDoesPluginDependGameplayTagSource::No;
+	}
+
+	// If there's no source data for the tag then the tag is unknown and hasn't been registered yet.
+	return EDoesPluginDependGameplayTagSource::UnknownTag;
+}
+
+/*static*/ bool SPluginAuditBrowser::IsTagOnlyAvailableFromExcludedSources(const UGameplayTagsManager& Manager, const FGameplayTag& Tag, const TArray<FGameFeaturePlugin>& ExcludedPlugins)
+{
+	FString Comment;
+	TArray<FName> TagSources;
+	bool bIsTagExplicit, bIsRestrictedTag, bAllowNonRestrictedChildren;
+	if (Manager.GetTagEditorData(Tag.GetTagName(), Comment, TagSources, bIsTagExplicit, bIsRestrictedTag, bAllowNonRestrictedChildren))
+	{
+		int ExcludedSourcesFound = 0;
+		
+		for (const FName& TagSourceName : TagSources)
+		{
+			const FString& TagSourceString = TagSourceName.ToString();
+
+			FName TagSourceIniPackage;
+			if (TagSourceString.EndsWith(TEXT(".ini")))
+			{
+				if (const FGameplayTagSource* TagSource = Manager.FindTagSource(TagSourceName))
+				{
+					if (ensure(TagSource->SourceTagList))
+					{
+						FString ContentFilePath = FPaths::GetPath(TagSource->SourceTagList->ConfigFileName) / TEXT("../../Content/");
+						FString RootContentPath;
+						if (FPackageName::TryConvertFilenameToLongPackageName(ContentFilePath, RootContentPath))
+						{
+							TagSourceIniPackage = FName(*RootContentPath);
+						}
+					}
+				}
+			}
+			
+			for (const FGameFeaturePlugin& ExcludedPlugin : ExcludedPlugins)
+			{
+				if (ExcludedPlugin.ModuleNames.Contains(TagSourceString))
+				{
+					ExcludedSourcesFound++;
+					break;
+				}
+
+				if (TagSourceString.StartsWith(ExcludedPlugin.ContentRoot.ToString()))
+				{
+					ExcludedSourcesFound++;
+					break;
+				}
+
+				if (ExcludedPlugin.ContentRoot == TagSourceIniPackage)
+				{
+					ExcludedSourcesFound++;
+					break;
+				}
+			}
+		}
+		
+		if (ExcludedSourcesFound == TagSources.Num())
+		{
+			return true;
+		}
+	}
+
+	return false;
+}
+
+/*static*/ void SPluginAuditBrowser::GetGameFeaturePlugins(const TArray<TSharedRef<IPlugin>>& Plugins, TArray<FGameFeaturePlugin>& GameFeaturePlugins)
+{
+	const FString BuiltInGameFeaturePluginsFolder = FPaths::ConvertRelativePathToFull(FPaths::ProjectPluginsDir() + TEXT("GameFeatures/"));
+	for (const TSharedRef<IPlugin>& Plugin : Plugins)
+	{
+		const FString& PluginDescriptorFilename = Plugin->GetDescriptorFileName();
+		if (!PluginDescriptorFilename.IsEmpty() && FPaths::ConvertRelativePathToFull(PluginDescriptorFilename).StartsWith(BuiltInGameFeaturePluginsFolder))
+		{
+			FGameFeaturePlugin GameFeaturePlugin;
+			GameFeaturePlugin.Plugin = Plugin;
+			for (const FModuleDescriptor& Module : Plugin->GetDescriptor().Modules)
+			{
+				const FString ModuleName = Module.Name.ToString();
+				GameFeaturePlugin.ModuleNames.Add(ModuleName);
+				GameFeaturePlugin.ScriptPackages.Add(FPackageName::GetModuleScriptPackageName(ModuleName));
+			}
+			GameFeaturePlugin.ContentRoot = FName(*FString::Printf(TEXT("/%s/"), *FPaths::GetBaseFilename(PluginDescriptorFilename)));
+
+			GameFeaturePlugins.Add(MoveTemp(GameFeaturePlugin));
+		}
+	}
+}
+
+#undef LOCTEXT_NAMESPACE
diff --git a/Engine/Plugins/Experimental/PluginAudit/Source/Private/SPluginAuditBrowser.h b/Engine/Plugins/Experimental/PluginAudit/Source/Private/SPluginAuditBrowser.h
new file mode 100644
--- /dev/null
+++ b/Engine/Plugins/Experimental/PluginAudit/Source/Private/SPluginAuditBrowser.h
@@ -0,0 +1,77 @@
+// Copyright Epic Games, Inc. All Rights Reserved.
+
+#pragma once
+
+#include "Widgets/SCompoundWidget.h"
+#include "Tasks/Task.h"
+#include "AssetRegistry/AssetIdentifier.h"
+
+//////////////////////////////////////////////////////////////////////////
+// SPluginAuditBrowser
+
+class IPlugin;
+struct FGameFeaturePlugin;
+class UGameplayTagsManager;
+struct FGameplayTag;
+class IMessageLogListing;
+class IMessageToken;
+
+class SPluginAuditBrowser : public SCompoundWidget
+{
+public:
+	SLATE_BEGIN_ARGS(SPluginAuditBrowser)
+	{}
+
+	SLATE_END_ARGS()
+public:
+	void Construct(const FArguments& InArgs);
+
+private:
+	void CreateLogListing();
+	void BuildPluginList();
+	void RefreshToolBar();
+
+	class FCookedPlugin
+	{
+	public:
+		FCookedPlugin(const TSharedRef<IPlugin>& InPlugin)
+			: Plugin(InPlugin)
+		{}
+		virtual ~FCookedPlugin() = default;
+
+		TSharedRef<IPlugin> Plugin;
+		bool bSimulateDisabled = false;
+	};
+
+	TSharedRef<ITableRow> MakeCookedPluginRow(TSharedRef<FCookedPlugin> InItem, const TSharedRef<STableViewBase>& OwnerTable);
+
+	struct FGameFeaturePlugin
+	{
+		TSharedPtr<IPlugin> Plugin;
+		TArray<FString> ModuleNames;
+		TArray<FString> ScriptPackages;
+		FName ContentRoot;
+	};
+
+	enum class EDoesPluginDependGameplayTagSource
+	{
+		Yes,
+		No,
+		UnknownTag
+	};
+
+	void RefreshViolations();
+
+	static TArray<TSharedRef<FTokenizedMessage>> ScanForViolations(TArray<TSharedRef<IPlugin>> InIncludedGameFeaturePlugins, TArray<TSharedRef<IPlugin>> InExcludedGameFeaturePlugins);
+	static TArray<TSharedPtr<IPlugin>> GetTagSourcePlugins(const UGameplayTagsManager& Manager, FName TagName);
+	static EDoesPluginDependGameplayTagSource DoesPluginDependOnGameplayTagSource(const UGameplayTagsManager& Manager, const TSharedPtr<IPlugin>& DependentPlugin, FName TagName, TArray<TSharedPtr<IPlugin>>& OutPossibleSources);
+	static bool IsTagOnlyAvailableFromExcludedSources(const UGameplayTagsManager& Manager, const FGameplayTag& Tag, const TArray<FGameFeaturePlugin>& ExcludedPlugins);
+	static void GetGameFeaturePlugins(const TArray<TSharedRef<IPlugin>>& Plugins, TArray<FGameFeaturePlugin>& GameFeaturePlugins);
+
+private:
+	TArray<TSharedRef<IPlugin>> IncludedGameFeaturePlugins;
+	TArray<TSharedRef<IPlugin>> ExcludedGameFeaturePlugins;
+
+	TArray<TSharedRef<FCookedPlugin>> CookedPlugins;
+	TSharedPtr<IMessageLogListing> LogListing;
+};
