Fixing a linux build issue.

diff --git a/Engine/Plugins/Experimental/PluginAudit/Source/Private/SPluginAuditBrowser.cpp b/Engine/Plugins/Experimental/PluginAudit/Source/Private/SPluginAuditBrowser.cpp
--- a/Engine/Plugins/Experimental/PluginAudit/Source/Private/SPluginAuditBrowser.cpp
+++ b/Engine/Plugins/Experimental/PluginAudit/Source/Private/SPluginAuditBrowser.cpp
@@ -174,7 +174,7 @@ void SPluginAuditBrowser::RefreshViolations()
 {
 	TArray<TSharedRef<IPlugin>> IncludedPlugins = IncludedGameFeaturePlugins;
 	TArray<TSharedRef<IPlugin>> ExcludedPlugins = ExcludedGameFeaturePlugins;
-	for (const TSharedPtr<FCookedPlugin>& CookedPlugin : CookedPlugins)
+	for (const TSharedRef<FCookedPlugin>& CookedPlugin : CookedPlugins)
 	{
 		if (CookedPlugin->bSimulateDisabled)
 		{
