Plugin audit unity build issue.

diff --git a/Engine/Plugins/Experimental/PluginAudit/Source/Private/SPluginAuditBrowser.h b/Engine/Plugins/Experimental/PluginAudit/Source/Private/SPluginAuditBrowser.h
--- a/Engine/Plugins/Experimental/PluginAudit/Source/Private/SPluginAuditBrowser.h
+++ b/Engine/Plugins/Experimental/PluginAudit/Source/Private/SPluginAuditBrowser.h
@@ -15,6 +15,8 @@ class UGameplayTagsManager;
 struct FGameplayTag;
 class IMessageLogListing;
 class IMessageToken;
+class STableViewBase;
+class ITableRow;
 
 class SPluginAuditBrowser : public SCompoundWidget
 {
