#UE  Minor plugin audit browser optimization

diff --git a/Engine/Plugins/Experimental/PluginAudit/Source/Private/SPluginAuditBrowser.cpp b/Engine/Plugins/Experimental/PluginAudit/Source/Private/SPluginAuditBrowser.cpp
--- a/Engine/Plugins/Experimental/PluginAudit/Source/Private/SPluginAuditBrowser.cpp
+++ b/Engine/Plugins/Experimental/PluginAudit/Source/Private/SPluginAuditBrowser.cpp
@@ -280,7 +280,7 @@ void SPluginAuditBrowser::RefreshViolations()
 							);
 							if (DependencyResult == EDoesPluginDependOnGameplayTagSource::UnknownTag)
 							{
-								Message->AddToken(FTextToken::Create(LOCTEXT("FromAndUnknownPlugin", " from a plug-in.  The gameplay tag's source is unknown so it's probably in a plugin that's not registered as a dependency of this plug-in, and nothing else loads the plug-in so we can't figure out where it's suppose to come from.")));
+								Message->AddToken(FTextToken::Create(LOCTEXT("FromAndUnknownPlugin", " from a plug-in.  The gameplay tag's source is unknown so it's probably in a plugin that's not registered as a dependency of this plug-in, and nothing else loads the plug-in so we can't figure out where it's supposed to come from.")));
 							}
 							else
 							{
@@ -452,10 +452,13 @@ void SPluginAuditBrowser::RefreshViolations()
 			break;
 		}
 
+		if (!TagPackageName.IsNone())
+		{
+			FString TagPackageNameString = TagPackageName.ToString();
 		for (const TSharedRef<IPlugin>& Plugin : AllPlugins)
 		{
 			FString ContentRoot = FString::Printf(TEXT("/%s/"), *FPaths::GetBaseFilename(Plugin->GetDescriptorFileName()));
-			if (TagPackageName.ToString().StartsWith(ContentRoot))
+				if (TagPackageNameString.StartsWith(ContentRoot))
 			{
 				SourcePlugins.Add(Plugin);
 				continue;
@@ -478,6 +481,7 @@ void SPluginAuditBrowser::RefreshViolations()
 			}
 		}
 	}
+	}
 
 	return SourcePlugins;
 }
