Added a message token to the `plugin name` text to open the plugin properties window

diff --git a/Engine/Plugins/Experimental/PluginAudit/Source/Private/SPluginAuditBrowser.cpp b/Engine/Plugins/Experimental/PluginAudit/Source/Private/SPluginAuditBrowser.cpp
--- a/Engine/Plugins/Experimental/PluginAudit/Source/Private/SPluginAuditBrowser.cpp
+++ b/Engine/Plugins/Experimental/PluginAudit/Source/Private/SPluginAuditBrowser.cpp
@@ -2,6 +2,9 @@
 
 #include "SPluginAuditBrowser.h"
 
+#include "Features/EditorFeatures.h"
+#include "Features/IModularFeatures.h"
+#include "Features/IPluginsEditorFeature.h"
 #include "GameFeaturesProjectPolicies.h"
 #include "GameFeaturesSubsystem.h"
 #include "GameplayTagsManager.h"
@@ -264,9 +267,17 @@ void SPluginAuditBrowser::RefreshViolations()
 							TSharedRef<FTokenizedMessage> Message = FTokenizedMessage::Create(EMessageSeverity::Error);
 							Message->AddToken(FTextToken::Create(LOCTEXT("GameplayTagReference", "Gameplay Tag Reference")));
 							Message->AddToken(FTextToken::Create(FText::FromString(TEXT(":"))));
-							Message->AddToken(FTextToken::Create(FText::FormatNamed(LOCTEXT("DependingGameplayTagUnknown", "The {ReferencerPlugin} contains "), 
-								TEXT("ReferencerPlugin"), FText::FromString(ReferencerPlugin->GetName())
-							)));
+							Message->AddToken(FTextToken::Create(LOCTEXT("ThePlugin", "The Plugin ")));
+
+							Message->AddToken(
+								FAssetNameToken::Create(ReferencerPlugin->GetName())->OnMessageTokenActivated(FOnMessageTokenActivated::CreateLambda([ReferencerPlugin](const TSharedRef<class IMessageToken>&) {
+									IPluginsEditorFeature& PluginEditor = IModularFeatures::Get().GetModularFeature<IPluginsEditorFeature>(EditorFeatures::PluginsEditor);
+									PluginEditor.OpenPluginEditor(ReferencerPlugin.ToSharedRef(), nullptr, FSimpleDelegate());
+									}))
+							);
+
+							Message->AddToken(FTextToken::Create(LOCTEXT("ThePluginContains", " contains ")));
+
 							Message->AddToken(
 								FAssetNameToken::Create(Referencer.ToString())->OnMessageTokenActivated(FOnMessageTokenActivated::CreateLambda([Referencer](const TSharedRef<class IMessageToken>&) {
 									IAssetManagerEditorModule::Get().OpenReferenceViewerUI({ Referencer });
