Added a checkbox to toggle the `Simulate Disabled` state for all plugins

diff --git a/Engine/Plugins/Experimental/PluginAudit/Source/Private/SPluginAuditBrowser.cpp b/Engine/Plugins/Experimental/PluginAudit/Source/Private/SPluginAuditBrowser.cpp
--- a/Engine/Plugins/Experimental/PluginAudit/Source/Private/SPluginAuditBrowser.cpp
+++ b/Engine/Plugins/Experimental/PluginAudit/Source/Private/SPluginAuditBrowser.cpp
@@ -65,9 +65,23 @@ void SPluginAuditBrowser::Construct(const FArguments& InArgs)
 			+ SSplitter::Slot()
 			.Value(0.30f)
 			[
-				SNew(SListView< TSharedRef<FCookedPlugin> >)
-				.ListItemsSource(&CookedPlugins)
-				.OnGenerateRow(this, &SPluginAuditBrowser::MakeCookedPluginRow)
+				SNew(SVerticalBox)
+				+ SVerticalBox::Slot()
+				.AutoHeight()
+				.Padding(0.f, 4.0f, 0.0f, 6.0f)
+				[
+					SNew(SCheckBox)
+					.ToolTipText(LOCTEXT("ToggleAll", "Toggle `Simulate Disabled` for all plugins"))
+					.IsChecked(ECheckBoxState::Checked)
+					.OnCheckStateChanged(this, &SPluginAuditBrowser::OnGlobalCheckboxStateChanged)
+				]
+
+				+ SVerticalBox::Slot()
+				[
+					SNew(SListView< TSharedRef<FCookedPlugin> >)
+					.ListItemsSource(&CookedPlugins)
+					.OnGenerateRow(this, &SPluginAuditBrowser::MakeCookedPluginRow)
+				]
 			]
 			+ SSplitter::Slot()
 			.Value(0.70f)
@@ -135,6 +149,14 @@ void SPluginAuditBrowser::RefreshToolBar()
 	}
 }
 
+void SPluginAuditBrowser::OnGlobalCheckboxStateChanged(ECheckBoxState State)
+{
+	for (auto& Item : CookedPlugins)
+	{
+		Item->bSimulateDisabled = (State == ECheckBoxState::Checked) ? false : true;
+	}
+}
+
 TSharedRef<ITableRow> SPluginAuditBrowser::MakeCookedPluginRow(TSharedRef<SPluginAuditBrowser::FCookedPlugin> InItem, const TSharedRef<STableViewBase>& OwnerTable)
 {
 	return SNew(STableRow< TSharedRef<FCookedPlugin> >, OwnerTable)
diff --git a/Engine/Plugins/Experimental/PluginAudit/Source/Private/SPluginAuditBrowser.h b/Engine/Plugins/Experimental/PluginAudit/Source/Private/SPluginAuditBrowser.h
--- a/Engine/Plugins/Experimental/PluginAudit/Source/Private/SPluginAuditBrowser.h
+++ b/Engine/Plugins/Experimental/PluginAudit/Source/Private/SPluginAuditBrowser.h
@@ -34,6 +34,8 @@ private:
 	void BuildPluginList();
 	void RefreshToolBar();
 
+	void OnGlobalCheckboxStateChanged(ECheckBoxState State);
+
 	class FCookedPlugin
 	{
 	public:
