Fixing compile error from missing include for ECheckBoxState

diff --git a/Engine/Plugins/Experimental/PluginAudit/Source/Private/SPluginAuditBrowser.h b/Engine/Plugins/Experimental/PluginAudit/Source/Private/SPluginAuditBrowser.h
--- a/Engine/Plugins/Experimental/PluginAudit/Source/Private/SPluginAuditBrowser.h
+++ b/Engine/Plugins/Experimental/PluginAudit/Source/Private/SPluginAuditBrowser.h
@@ -5,6 +5,7 @@
 #include "Widgets/SCompoundWidget.h"
 #include "Tasks/Task.h"
 #include "AssetRegistry/AssetIdentifier.h"
+#include "Styling/SlateTypes.h"
 
 //////////////////////////////////////////////////////////////////////////
 // SPluginAuditBrowser
