Added a Search box to filter plugins by name. The global toggle will now only apply to visible plugins

diff --git a/Engine/Plugins/Experimental/PluginAudit/Source/Private/SPluginAuditBrowser.cpp b/Engine/Plugins/Experimental/PluginAudit/Source/Private/SPluginAuditBrowser.cpp
--- a/Engine/Plugins/Experimental/PluginAudit/Source/Private/SPluginAuditBrowser.cpp
+++ b/Engine/Plugins/Experimental/PluginAudit/Source/Private/SPluginAuditBrowser.cpp
@@ -5,6 +5,7 @@
 #include "Features/EditorFeatures.h"
 #include "Features/IModularFeatures.h"
 #include "Features/IPluginsEditorFeature.h"
+#include "Filters/SFilterSearchBox.h"
 #include "GameFeaturesProjectPolicies.h"
 #include "GameFeaturesSubsystem.h"
 #include "GameplayTagsManager.h"
@@ -19,11 +20,11 @@
 #include "IMessageLogListing.h"
 #include "MessageLogModule.h"
 #include "Misc/UObjectToken.h"
+#include "Widgets/Input/SSearchBox.h"
 #include "Widgets/SBoxPanel.h"
 #include "Widgets/Layout/SBorder.h"
 #include "Widgets/Layout/SSplitter.h"
 #include "Widgets/Input/SCheckBox.h"
-#include "Widgets/Views/SListView.h"
 #include "AssetManagerEditorModule.h"
 #include "ToolMenus.h"
 
@@ -36,13 +37,29 @@ namespace PluginAudit
 
 void SPluginAuditBrowser::Construct(const FArguments& InArgs)
 {
-	CreateLogListing();
 	BuildPluginList();
+	CreateLogListing();
+
+	// Setup text filtering
+	PluginTextFilter = MakeShareable(new FCookedPluginTextFilter(FCookedPluginTextFilter::FItemToStringArray::CreateLambda([](const IPlugin* Plugin, TArray<FString>& OutStringArray)
+		{
+			OutStringArray.Add(Plugin->GetFriendlyName());
+		})
+	));
+
+	PluginTextFilter->OnChanged().AddSP(this, &SPluginAuditBrowser::OnPluginTextFilterChanged);
+
+	RebuildAndFilterPluginList();
 
 	FMessageLogModule& MessageLogModule = FModuleManager::LoadModuleChecked<FMessageLogModule>("MessageLog");
 
 	RefreshToolBar();
 
+	PluginListView =
+		SNew(SListView< TSharedRef<FCookedPlugin> >)
+		.ListItemsSource(&FilteredCookedPlugins)
+		.OnGenerateRow(this, &SPluginAuditBrowser::MakeCookedPluginRow);
+
 	ChildSlot
 	[
 		SNew(SVerticalBox)
@@ -65,22 +82,36 @@ void SPluginAuditBrowser::Construct(const FArguments& InArgs)
 			+ SSplitter::Slot()
 			.Value(0.30f)
 			[
-				SNew(SVerticalBox)
-				+ SVerticalBox::Slot()
-				.AutoHeight()
-				.Padding(0.f, 4.0f, 0.0f, 6.0f)
-				[
-					SNew(SCheckBox)
-					.ToolTipText(LOCTEXT("ToggleAll", "Toggle `Simulate Disabled` for all plugins"))
-					.IsChecked(ECheckBoxState::Checked)
-					.OnCheckStateChanged(this, &SPluginAuditBrowser::OnGlobalCheckboxStateChanged)
-				]
-
-				+ SVerticalBox::Slot()
+				SNew(SBorder)
+				.Padding(FMargin(3))
+				.BorderImage(FAppStyle::GetBrush("ToolPanel.GroupBorder"))
 				[
-					SNew(SListView< TSharedRef<FCookedPlugin> >)
-					.ListItemsSource(&CookedPlugins)
-					.OnGenerateRow(this, &SPluginAuditBrowser::MakeCookedPluginRow)
+					SNew(SVerticalBox)
+					+ SVerticalBox::Slot()
+					.AutoHeight()
+					.Padding(0.f, 4.0f, 0.0f, 6.0f)
+					[
+						SNew(SHorizontalBox)
+						+ SHorizontalBox::Slot()
+						.AutoWidth()				
+						[
+							SNew(SCheckBox)
+							.ToolTipText(LOCTEXT("ToggleAll", "Toggle disabled for all visible plugins"))
+							.IsChecked(this, &SPluginAuditBrowser::GetGlobalDisabledState)
+							.OnCheckStateChanged(this, &SPluginAuditBrowser::OnGlobalDisabledStateChanged)
+						]
+						+ SHorizontalBox::Slot()
+						.FillWidth(1.0f)
+						[
+							SAssignNew(SearchBoxPtr, SSearchBox)
+							.OnTextChanged(this, &SPluginAuditBrowser::SearchBox_OnPluginSearchTextChanged)
+						]
+					]
+
+					+ SVerticalBox::Slot()
+					[
+						PluginListView.ToSharedRef()
+					]
 				]
 			]
 			+ SSplitter::Slot()
@@ -122,8 +153,21 @@ void SPluginAuditBrowser::BuildPluginList()
 	{
 		CookedPlugins.Add(MakeShared<FCookedPlugin>(Plugin));
 	}
+}
+
+void SPluginAuditBrowser::RebuildAndFilterPluginList()
+{
+	FilteredCookedPlugins.Empty();
+
+	for (const auto& CookedPlugin : CookedPlugins)
+	{
+		if (PluginTextFilter->PassesFilter(&CookedPlugin->Plugin.Get()))
+		{
+			FilteredCookedPlugins.Add(CookedPlugin);
+		}
+	}
 
-	CookedPlugins.Sort([](const TSharedRef<FCookedPlugin>& A, const TSharedRef<FCookedPlugin>& B)
+	FilteredCookedPlugins.Sort([](const TSharedRef<FCookedPlugin>& A, const TSharedRef<FCookedPlugin>& B)
 	{
 		return A->Plugin->GetName() < B->Plugin->GetName();
 	});
@@ -149,14 +193,45 @@ void SPluginAuditBrowser::RefreshToolBar()
 	}
 }
 
-void SPluginAuditBrowser::OnGlobalCheckboxStateChanged(ECheckBoxState State)
+void SPluginAuditBrowser::OnPluginTextFilterChanged()
 {
-	for (auto& Item : CookedPlugins)
+	RebuildAndFilterPluginList();
+
+	PluginListView->RequestListRefresh();
+}
+
+void SPluginAuditBrowser::SearchBox_OnPluginSearchTextChanged(const FText& NewText)
+{
+	PluginTextFilter->SetRawFilterText(NewText);
+	SearchBoxPtr->SetError(PluginTextFilter->GetFilterErrorText());
+}
+
+void SPluginAuditBrowser::OnGlobalDisabledStateChanged(ECheckBoxState State)
+{
+	bGlobalDisabledState = (State == ECheckBoxState::Checked) ? false : true;
+	
+	for (auto& Item : FilteredCookedPlugins)
 	{
-		Item->bSimulateDisabled = (State == ECheckBoxState::Checked) ? false : true;
+		Item->bSimulateDisabled = bGlobalDisabledState;
 	}
 }
 
+ECheckBoxState SPluginAuditBrowser::GetGlobalDisabledState() const
+{
+	bool bStateMismatch = false;
+	for (const auto& Item : FilteredCookedPlugins)
+	{
+		if (Item->bSimulateDisabled != bGlobalDisabledState)
+		{
+			bStateMismatch = true;
+			break;
+		}
+	}
+
+	return bStateMismatch ? ECheckBoxState::Undetermined :
+		bGlobalDisabledState ? ECheckBoxState::Unchecked : ECheckBoxState::Checked;
+}
+
 TSharedRef<ITableRow> SPluginAuditBrowser::MakeCookedPluginRow(TSharedRef<SPluginAuditBrowser::FCookedPlugin> InItem, const TSharedRef<STableViewBase>& OwnerTable)
 {
 	return SNew(STableRow< TSharedRef<FCookedPlugin> >, OwnerTable)
diff --git a/Engine/Plugins/Experimental/PluginAudit/Source/Private/SPluginAuditBrowser.h b/Engine/Plugins/Experimental/PluginAudit/Source/Private/SPluginAuditBrowser.h
--- a/Engine/Plugins/Experimental/PluginAudit/Source/Private/SPluginAuditBrowser.h
+++ b/Engine/Plugins/Experimental/PluginAudit/Source/Private/SPluginAuditBrowser.h
@@ -5,6 +5,7 @@
 #include "Widgets/SCompoundWidget.h"
 #include "Tasks/Task.h"
 #include "AssetRegistry/AssetIdentifier.h"
+#include "Misc/TextFilter.h"
 #include "Styling/SlateTypes.h"
 
 //////////////////////////////////////////////////////////////////////////
@@ -16,6 +17,7 @@ class UGameplayTagsManager;
 struct FGameplayTag;
 class IMessageLogListing;
 class IMessageToken;
+class SFilterSearchBox;
 class STableViewBase;
 class ITableRow;
 class FTokenizedMessage;
@@ -33,9 +35,13 @@ public:
 private:
 	void CreateLogListing();
 	void BuildPluginList();
+	void RebuildAndFilterPluginList();
 	void RefreshToolBar();
 
-	void OnGlobalCheckboxStateChanged(ECheckBoxState State);
+	void OnPluginTextFilterChanged();
+	void SearchBox_OnPluginSearchTextChanged(const FText& NewText);
+	ECheckBoxState GetGlobalDisabledState() const;
+	void OnGlobalDisabledStateChanged(ECheckBoxState State);
 
 	class FCookedPlugin
 	{
@@ -75,9 +81,24 @@ private:
 	static void GetGameFeaturePlugins(const TArray<TSharedRef<IPlugin>>& Plugins, TArray<FGameFeaturePlugin>& GameFeaturePlugins);
 
 private:
+	typedef TTextFilter<const IPlugin*> FCookedPluginTextFilter;
+
 	TArray<TSharedRef<IPlugin>> IncludedGameFeaturePlugins;
 	TArray<TSharedRef<IPlugin>> ExcludedGameFeaturePlugins;
 
 	TArray<TSharedRef<FCookedPlugin>> CookedPlugins;
+	TArray<TSharedRef<FCookedPlugin>> FilteredCookedPlugins;
+
 	TSharedPtr<IMessageLogListing> LogListing;
+
+	/** The list view widget for our plugins list */
+	TSharedPtr<SListView<TSharedRef<FCookedPlugin>>> PluginListView;
+
+	/** The plugin search box widget */
+	TSharedPtr<class SSearchBox> SearchBoxPtr;
+
+	/** Text filter object for typing in filter text to the search box */
+	TSharedPtr<FCookedPluginTextFilter> PluginTextFilter;
+
+	bool bGlobalDisabledState = false;
 };
