Fixed compile error with missing SListView header

diff --git a/Engine/Plugins/Experimental/PluginAudit/Source/Private/SPluginAuditBrowser.h b/Engine/Plugins/Experimental/PluginAudit/Source/Private/SPluginAuditBrowser.h
--- a/Engine/Plugins/Experimental/PluginAudit/Source/Private/SPluginAuditBrowser.h
+++ b/Engine/Plugins/Experimental/PluginAudit/Source/Private/SPluginAuditBrowser.h
@@ -7,6 +7,7 @@
 #include "AssetRegistry/AssetIdentifier.h"
 #include "Misc/TextFilter.h"
 #include "Styling/SlateTypes.h"
+#include "Widgets/Views/SListView.h"
 
 //////////////////////////////////////////////////////////////////////////
 // SPluginAuditBrowser
