Added an experimental plugin dependency viewer graph into the plugin audit tool

diff --git a/Engine/Plugins/Experimental/PluginAudit/Source/Private/EdGraphNode_PluginReference.cpp b/Engine/Plugins/Experimental/PluginAudit/Source/Private/EdGraphNode_PluginReference.cpp
new file mode 100644
--- /dev/null
+++ b/Engine/Plugins/Experimental/PluginAudit/Source/Private/EdGraphNode_PluginReference.cpp
@@ -0,0 +1,127 @@
+// Copyright Epic Games, Inc. All Rights Reserved.
+
+#include "EdGraphNode_PluginReference.h"
+#include "EdGraph/EdGraphPin.h"
+#include "Interfaces/IPluginManager.h"
+#include "PluginReferencePinCategory.h"
+#include "Styling/AppStyle.h"
+
+#include UE_INLINE_GENERATED_CPP_BY_NAME(EdGraphNode_PluginReference)
+
+#define LOCTEXT_NAMESPACE "PluginReferenceViewer"
+
+//////////////////////////////////////////////////////////////////////////
+// UEdGraphNode_Reference
+
+UEdGraphNode_PluginReference::UEdGraphNode_PluginReference(const FObjectInitializer& ObjectInitializer)
+	: Super(ObjectInitializer)
+	, bIsEnginePlugin(false)
+	, DependencyPin(nullptr)
+	, ReferencerPin(nullptr)
+{
+}
+
+FPluginIdentifier UEdGraphNode_PluginReference::GetIdentifier() const
+{
+	return PluginIdentifier;
+}
+
+void UEdGraphNode_PluginReference::SetupPluginReferenceNode(const FIntPoint& InNodeLoc, const FPluginIdentifier InPluginIdentifier, const TSharedPtr<const IPlugin>& InPlugin, bool bInAllowThumbnail)
+{
+	NodePosX = InNodeLoc.X;
+	NodePosY = InNodeLoc.Y;
+
+	AssetBrush = FSlateIcon(FName("PluginStyle"), "Plugins.TabIcon");
+
+	PluginIdentifier = InPluginIdentifier;
+	NodeTitle = FText::FromString(InPlugin->GetName());
+
+	bAllowThumbnail = bInAllowThumbnail;
+	bIsEnginePlugin = (InPlugin->GetType() == EPluginType::Engine) ? true : false;
+
+	CachedPlugin = InPlugin;
+
+	AllocateDefaultPins();
+}
+
+bool UEdGraphNode_PluginReference::AllowsThumbnail() const
+{
+	return bAllowThumbnail;
+}
+
+TSharedPtr<const IPlugin> UEdGraphNode_PluginReference::GetPlugin() const
+{
+	return CachedPlugin;
+}
+
+UEdGraph_PluginReferenceViewer* UEdGraphNode_PluginReference::GetPluginReferenceViewerGraph() const
+{
+	return Cast<UEdGraph_PluginReferenceViewer>(GetGraph());
+}
+
+void UEdGraphNode_PluginReference::AddReferencer(UEdGraphNode_PluginReference* ReferencerNode)
+{
+	UEdGraphPin* ReferencerDependencyPin = ReferencerNode->GetDependencyPin();
+
+	if (ensure(ReferencerDependencyPin))
+	{
+		ReferencerDependencyPin->bHidden = false;
+		ReferencerPin->bHidden = false;
+		ReferencerPin->MakeLinkTo(ReferencerDependencyPin);
+	}
+}
+
+void UEdGraphNode_PluginReference::AllocateDefaultPins()
+{
+	ReferencerPin = CreatePin(EEdGraphPinDirection::EGPD_Input, NAME_None, NAME_None);
+	DependencyPin = CreatePin(EEdGraphPinDirection::EGPD_Output, NAME_None, NAME_None);
+
+	ReferencerPin->bHidden = true;
+	FName PassiveName = PluginReferencePinUtil::GetName(EPluginReferencePinCategory::LinkEndPassive);
+	ReferencerPin->PinType.PinCategory = PassiveName;
+	DependencyPin->bHidden = true;
+	DependencyPin->PinType.PinCategory = PassiveName;
+}
+
+FSlateIcon UEdGraphNode_PluginReference::GetIconAndTint(FLinearColor& OutColor) const
+{
+	OutColor = FLinearColor::White;
+	return AssetBrush;
+}
+
+FText UEdGraphNode_PluginReference::GetNodeTitle(ENodeTitleType::Type TitleType) const
+{
+	return NodeTitle;
+}
+
+FLinearColor UEdGraphNode_PluginReference::GetNodeTitleColor() const
+{
+	if (bIsEnginePlugin)
+	{
+		return FLinearColor(0.55f, 0.55f, 0.55f);
+	}
+	else
+	{
+		return FLinearColor(0.0f, 0.42f, 1.0f);
+	}
+}
+
+FText UEdGraphNode_PluginReference::GetTooltipText() const
+{
+	FString TooltipText = CachedPlugin->GetBaseDir();
+	FPaths::MakePathRelativeTo(TooltipText, *FPaths::RootDir());
+	
+	return FText::FromString(TooltipText);
+}
+
+UEdGraphPin* UEdGraphNode_PluginReference::GetDependencyPin()
+{
+	return DependencyPin;
+}
+
+UEdGraphPin* UEdGraphNode_PluginReference::GetReferencerPin()
+{
+	return ReferencerPin;
+}
+
+#undef LOCTEXT_NAMESPACE
\ No newline at end of file
diff --git a/Engine/Plugins/Experimental/PluginAudit/Source/Private/EdGraphNode_PluginReference.h b/Engine/Plugins/Experimental/PluginAudit/Source/Private/EdGraphNode_PluginReference.h
new file mode 100644
--- /dev/null
+++ b/Engine/Plugins/Experimental/PluginAudit/Source/Private/EdGraphNode_PluginReference.h
@@ -0,0 +1,58 @@
+// Copyright Epic Games, Inc. All Rights Reserved.
+
+#pragma once
+
+#include "CoreMinimal.h"
+#include "UObject/ObjectMacros.h"
+#include "EdGraph/EdGraphNode.h"
+#include "EdGraph_PluginReferenceViewer.h"
+#include "Textures/SlateIcon.h"
+#include "EdGraphNode_PluginReference.generated.h"
+
+class IPlugin;
+class UEdGraphPin;
+class UEdGraph_PluginReferenceViewer;
+
+UCLASS()
+class UEdGraphNode_PluginReference : public UEdGraphNode
+{
+	GENERATED_UCLASS_BODY()
+
+	FPluginIdentifier GetIdentifier() const;
+
+	void SetupPluginReferenceNode(const FIntPoint& InNodeLoc, const FPluginIdentifier InPluginIdentifier, const TSharedPtr<const IPlugin>& InPlugin, bool bAllowThumbnail);
+
+	bool AllowsThumbnail() const;
+	TSharedPtr<const IPlugin> GetPlugin() const;
+
+	/** Returns our owning graph */
+	UEdGraph_PluginReferenceViewer* GetPluginReferenceViewerGraph() const;
+	void AddReferencer(UEdGraphNode_PluginReference* ReferencerNode);
+
+	// UEdGraphNode implementation
+	virtual void AllocateDefaultPins() override;
+	virtual FSlateIcon GetIconAndTint(FLinearColor& OutColor) const override;
+	virtual FText GetNodeTitle(ENodeTitleType::Type TitleType) const override;
+	virtual FLinearColor GetNodeTitleColor() const override;
+	virtual FText GetTooltipText() const override;
+	virtual bool ShowPaletteIconOnNode() const override { return true; }
+	// End UEdGraphNode implementation
+	
+	UEdGraphPin* GetDependencyPin();
+	UEdGraphPin* GetReferencerPin();
+
+private:
+	FPluginIdentifier PluginIdentifier;
+	FText NodeTitle;
+
+	bool bAllowThumbnail;
+	bool bIsEnginePlugin;
+
+	TSharedPtr<const IPlugin> CachedPlugin;
+	FSlateIcon AssetBrush;
+
+	UEdGraphPin* DependencyPin;
+	UEdGraphPin* ReferencerPin;
+
+	friend UEdGraph_PluginReferenceViewer;
+};
\ No newline at end of file
diff --git a/Engine/Plugins/Experimental/PluginAudit/Source/Private/EdGraph_PluginReferenceViewer.cpp b/Engine/Plugins/Experimental/PluginAudit/Source/Private/EdGraph_PluginReferenceViewer.cpp
new file mode 100644
--- /dev/null
+++ b/Engine/Plugins/Experimental/PluginAudit/Source/Private/EdGraph_PluginReferenceViewer.cpp
@@ -0,0 +1,383 @@
+// Copyright Epic Games, Inc. All Rights Reserved.
+
+#include "EdGraph_PluginReferenceViewer.h"
+
+#include "EdGraphNode_PluginReference.h"
+#include "EdGraph/EdGraphPin.h"
+#include "Interfaces/IPluginManager.h"
+#include "PluginReferencePinCategory.h"
+#include "SPluginReferenceViewer.h"
+
+#include UE_INLINE_GENERATED_CPP_BY_NAME(EdGraph_PluginReferenceViewer)
+
+bool ExceedsMaxSearchDepth(int32 CurrentDepth, int32 MaxDepth)
+{
+	const bool bIsWithinDepthLimits = MaxDepth > 0 && CurrentDepth < MaxDepth;
+	if (!bIsWithinDepthLimits)
+	{
+		return true;
+	}
+	return false;
+}
+
+FPluginReferenceNodeInfo::FPluginReferenceNodeInfo(const FPluginIdentifier& InIdentifier, bool bInReferencers)
+	: Identifier(InIdentifier)
+	, bReferencers(bInReferencers)
+	, ChildProvisionSize(0)
+{
+
+}
+
+bool FPluginReferenceNodeInfo::IsFirstParent(const FPluginIdentifier& InParentId) const
+{
+	return Parents.IsEmpty() || Parents[0] == InParentId;
+}
+
+int32 FPluginReferenceNodeInfo::ProvisionSize(const FPluginIdentifier& InParentId) const
+{
+	return IsFirstParent(InParentId) ? ChildProvisionSize : 1;
+}
+
+
+UEdGraph_PluginReferenceViewer::UEdGraph_PluginReferenceViewer(const FObjectInitializer& ObjectInitializer)
+	: Super(ObjectInitializer)
+{
+
+}
+
+void UEdGraph_PluginReferenceViewer::SetGraphRoot(const TArray<FPluginIdentifier>& GraphRootIdentifiers, const FIntPoint& GraphRootOrigin)
+{
+	CurrentGraphRootIdentifiers = GraphRootIdentifiers;
+	CurrentGraphRootOrigin = GraphRootOrigin;
+}
+
+void UEdGraph_PluginReferenceViewer::CachePluginDependencies(const TArray<TSharedRef<IPlugin>>& Plugins)
+{
+	PluginMap.Empty();
+
+	for (const TSharedRef<IPlugin>& Plugin : Plugins)
+	{
+		PluginMap.Add(FPluginIdentifier(Plugin->GetName()), Plugin);
+	}
+
+	struct FDependencyNodeBuilder
+	{
+		TMap<FPluginIdentifier, TUniquePtr<FPluginDependsNode>>& Nodes;
+
+		FDependencyNodeBuilder(TMap<FPluginIdentifier, TUniquePtr<FPluginDependsNode>>& InNodes)
+			: Nodes(InNodes)
+		{
+			Nodes.Empty();
+		}
+
+		void ProcessPlugins(const TArray<TSharedRef<IPlugin>>& Plugins)
+		{
+			for (const TSharedRef<IPlugin>& Plugin : Plugins)
+			{
+				Nodes.Add(FPluginIdentifier(Plugin->GetName()), MakeUnique<FPluginDependsNode>(Plugin->GetName()));
+			}
+
+			for (const TSharedRef<IPlugin>& Plugin : Plugins)
+			{
+				FPluginDependsNode* ParentNode = Nodes.FindChecked(FName(Plugin->GetName())).Get();
+
+				const FPluginDescriptor& PluginDescriptor = Plugin->GetDescriptor();
+				for (const FPluginReferenceDescriptor& Item : PluginDescriptor.Plugins)
+				{
+					FPluginDependsNode* ChildNode = nullptr;
+					TUniquePtr<FPluginDependsNode>* FoundNodePtr = Nodes.Find(FPluginIdentifier(Item.Name));
+					if (FoundNodePtr != nullptr)
+					{
+						ChildNode = FoundNodePtr->Get();
+					}
+					else
+					{
+						ChildNode = Nodes.Add(FPluginIdentifier(Item.Name), MakeUnique<FPluginDependsNode>(Item.Name)).Get();
+					}
+
+					ParentNode->Dependencies.Add(ChildNode);
+					ChildNode->Referencers.Add(ParentNode);
+				}
+			}
+		}
+	};
+
+	// Cache plugin dependencies
+	FDependencyNodeBuilder DependencyNodeBuilder(CachedDependsNodes);
+	DependencyNodeBuilder.ProcessPlugins(Plugins);
+}
+
+UEdGraphNode_PluginReference* UEdGraph_PluginReferenceViewer::ConstructNodes(const TArray<FPluginIdentifier>& GraphRootIdentifiers, const FIntPoint& GraphRootOrigin)
+{
+	const bool bIsCompactMode = PluginReferenceViewer.Pin()->IsCompactModeChecked();
+	const int32 SettingsReferencerMaxDepth = PluginReferenceViewer.Pin()->GetSearchReferencerDepthCount();
+	const int32 SettingsDependencyMaxDepth = PluginReferenceViewer.Pin()->GetSearchDependencyDepthCount();
+
+	TMap<FPluginIdentifier, FPluginReferenceNodeInfo> NewReferenceNodeInfos;
+	for (const FPluginIdentifier& RootIdentifier : GraphRootIdentifiers)
+	{
+		FPluginReferenceNodeInfo& RootNodeInfo = NewReferenceNodeInfos.FindOrAdd(RootIdentifier, FPluginReferenceNodeInfo(RootIdentifier, true));
+		RootNodeInfo.Parents.Emplace(NAME_None);
+	}
+
+	RecursivelyPopulateNodeInfos(true, GraphRootIdentifiers, NewReferenceNodeInfos, 0, SettingsReferencerMaxDepth);
+
+	TMap<FPluginIdentifier, FPluginReferenceNodeInfo> NewDependencyNodeInfos;
+	for (const FPluginIdentifier& RootIdentifier : GraphRootIdentifiers)
+	{
+		FPluginReferenceNodeInfo& DRootNodeInfo = NewDependencyNodeInfos.FindOrAdd(RootIdentifier, FPluginReferenceNodeInfo(RootIdentifier, false));
+		DRootNodeInfo.Parents.Emplace(NAME_None);
+	}
+
+	RecursivelyPopulateNodeInfos(false , GraphRootIdentifiers, NewDependencyNodeInfos, 0, SettingsDependencyMaxDepth);
+
+	// Store the Plugin in the NodeInfos
+	for (TPair<FPluginIdentifier, FPluginReferenceNodeInfo>& InfoPair : NewReferenceNodeInfos)
+	{
+		InfoPair.Value.Plugin = PluginMap.FindChecked(InfoPair.Key);
+	}
+
+	for (TPair<FPluginIdentifier, FPluginReferenceNodeInfo>& InfoPair : NewDependencyNodeInfos)
+	{
+		InfoPair.Value.Plugin = PluginMap.FindChecked(InfoPair.Key);
+	}
+
+	ReferencerNodeInfos = NewReferenceNodeInfos;
+	DependencyNodeInfos = NewDependencyNodeInfos;
+
+	RemoveAllNodes();
+
+	UEdGraphNode_PluginReference* RootNode = nullptr;
+
+	if (GraphRootIdentifiers.Num() > 0 && (!ReferencerNodeInfos.IsEmpty() || !DependencyNodeInfos.IsEmpty()))
+	{
+		FPluginIdentifier FirstGraphRootIdentifier = GraphRootIdentifiers[0];
+
+		const FPluginReferenceNodeInfo& NodeInfo = ReferencerNodeInfos[FirstGraphRootIdentifier];
+		RootNode = CreatePluginReferenceNode();
+		RootNode->SetupPluginReferenceNode(CurrentGraphRootOrigin, NodeInfo.Identifier, NodeInfo.Plugin, !bIsCompactMode);
+
+		RecursivelyCreateNodes(true, FirstGraphRootIdentifier, CurrentGraphRootOrigin, FirstGraphRootIdentifier, RootNode, ReferencerNodeInfos, 0, SettingsReferencerMaxDepth, /*bIsRoot*/ true);
+		
+		RecursivelyCreateNodes(false, FirstGraphRootIdentifier, CurrentGraphRootOrigin, FirstGraphRootIdentifier, RootNode, DependencyNodeInfos, 0, SettingsDependencyMaxDepth, /*bIsRoot*/ true);
+	}
+
+	return RootNode;
+}
+
+UEdGraphNode_PluginReference* UEdGraph_PluginReferenceViewer::RecursivelyCreateNodes(bool bInReferencers, const FPluginIdentifier& InPluginId, const FIntPoint& InNodeLoc, const FPluginIdentifier& InParentId, UEdGraphNode_PluginReference* InParentNode, TMap<FPluginIdentifier, FPluginReferenceNodeInfo>& InNodeInfos, int32 InCurrentDepth, int32 InMaxDepth, bool bIsRoot)
+{
+	check(InNodeInfos.Contains(InPluginId));
+
+	const FPluginReferenceNodeInfo& NodeInfo = InNodeInfos[InPluginId];
+	int32 NodeProvSize = 1;
+
+	UEdGraphNode_PluginReference* NewNode = nullptr;
+	if (bIsRoot)
+	{
+		NewNode = InParentNode;
+		NodeProvSize = NodeInfo.ProvisionSize(FPluginIdentifier(NAME_None));
+	}
+	else
+	{
+		const bool bIsCompactMode = PluginReferenceViewer.Pin()->IsCompactModeChecked();
+
+		NewNode = CreatePluginReferenceNode();
+		NewNode->SetupPluginReferenceNode(InNodeLoc, NodeInfo.Identifier, NodeInfo.Plugin, !bIsCompactMode);
+		NodeProvSize = NodeInfo.ProvisionSize(InParentId);
+	}
+
+	const bool bIsCompactMode = PluginReferenceViewer.Pin()->IsCompactModeChecked();
+
+	bool bIsFirstOccurance = bIsRoot || NodeInfo.IsFirstParent(InParentId);
+	FIntPoint ChildLoc = InNodeLoc;
+	if (!ExceedsMaxSearchDepth(InCurrentDepth, InMaxDepth) && bIsFirstOccurance) // Only expand the first parent
+	{
+		// position the children nodes
+		const int32 ColumnWidth = bIsCompactMode ? 500 : 800;
+		ChildLoc.X += bInReferencers ? -ColumnWidth : ColumnWidth;
+
+		int32 NodeSizeY = bIsCompactMode ? 100 : 200;
+		ChildLoc.Y -= (NodeProvSize - 1) * NodeSizeY * 0.5;
+
+		int32 Breadth = 0;
+		int32 ChildIdx = 0;
+		for (; ChildIdx < InNodeInfos[InPluginId].Children.Num(); ChildIdx++)
+		{
+			const TPair<FPluginIdentifier, EPluginReferencePinCategory>& Pair = InNodeInfos[InPluginId].Children[ChildIdx];
+
+			FPluginIdentifier ChildId = Pair.Key;
+			int32 ChildProvSize = 1;
+			if (InNodeInfos[ChildId].IsFirstParent(InPluginId))
+			{
+				ChildProvSize = InNodeInfos[ChildId].ProvisionSize(InPluginId);
+			}
+
+			// The provision size will always be at least 1 if it should be shown
+			if (ChildProvSize > 0)
+			{
+				ChildLoc.Y += (ChildProvSize - 1) * NodeSizeY * 0.5;
+
+				UEdGraphNode_PluginReference* ChildNode = RecursivelyCreateNodes(bInReferencers, ChildId, ChildLoc, InPluginId, NewNode, InNodeInfos, InCurrentDepth + 1, InMaxDepth);
+
+				if (bInReferencers)
+				{
+					ChildNode->GetDependencyPin()->PinType.PinCategory = PluginReferencePinUtil::GetName(Pair.Value);
+					NewNode->AddReferencer(ChildNode);
+				}
+				else
+				{
+					ChildNode->GetReferencerPin()->PinType.PinCategory = PluginReferencePinUtil::GetName(Pair.Value);
+					ChildNode->AddReferencer(NewNode);
+				}
+
+				ChildLoc.Y += NodeSizeY * (ChildProvSize + 1) * 0.5;
+				Breadth++;
+			}
+		}
+	}
+
+	return NewNode;
+}
+
+void UEdGraph_PluginReferenceViewer::GetSortedLinks(const TArray<FPluginIdentifier>& Identifiers, bool bReferencers, TMap<FPluginIdentifier, EPluginReferencePinCategory>& OutLinks) const
+{
+	const bool bShowEnginePlugins = PluginReferenceViewer.Pin()->IsShowEnginePluginsChecked();
+
+	TArray<FPluginDependency> LinksToAsset;
+
+	for (const FPluginIdentifier& PluginIdentifier : Identifiers)
+	{
+		LinksToAsset.Reset();
+		if (bReferencers)
+		{
+			GetPluginReferencers(PluginIdentifier, LinksToAsset);
+		}
+		else
+		{
+			GetPluginDependencies(PluginIdentifier, LinksToAsset);
+		}
+
+		for (const FPluginDependency& LinkToAsset : LinksToAsset)
+		{
+			EPluginReferencePinCategory& Category = OutLinks.FindOrAdd(LinkToAsset.Identifier, EPluginReferencePinCategory::LinkEndActive);
+			Category |= LinkToAsset.bIsEnabled ? EPluginReferencePinCategory::LinkTypeEnabled : EPluginReferencePinCategory::LinkTypeNone;
+			Category |= LinkToAsset.bIsOptional ? EPluginReferencePinCategory::LinkTypeOptional : EPluginReferencePinCategory::LinkTypeNone;
+		}
+	}
+}
+
+void UEdGraph_PluginReferenceViewer::RecursivelyPopulateNodeInfos(bool bInReferencers, const TArray<FPluginIdentifier>& Identifiers, TMap<FPluginIdentifier, FPluginReferenceNodeInfo>& InNodeInfos, int32 InCurrentDepth, int32 InMaxDepth)
+{
+	check(Identifiers.Num() > 0);
+	int32 ProvisionSize = 0;
+	const FPluginIdentifier PluginId = Identifiers[0];
+	if (!ExceedsMaxSearchDepth(InCurrentDepth, InMaxDepth))
+	{
+		TMap<FPluginIdentifier, EPluginReferencePinCategory> ReferenceLinks;
+		GetSortedLinks(Identifiers, bInReferencers, ReferenceLinks);
+
+		InNodeInfos[PluginId].Children.Reserve(ReferenceLinks.Num());
+		for (const TPair<FPluginIdentifier, EPluginReferencePinCategory>& Pair : ReferenceLinks)
+		{
+			FPluginIdentifier ChildId = Pair.Key;
+			if (!InNodeInfos.Contains(ChildId))
+			{
+				FPluginReferenceNodeInfo& NewNodeInfo = InNodeInfos.FindOrAdd(ChildId, FPluginReferenceNodeInfo(ChildId, bInReferencers));
+				InNodeInfos[ChildId].Parents.Emplace(PluginId);
+				InNodeInfos[PluginId].Children.Emplace(Pair);
+
+				RecursivelyPopulateNodeInfos(bInReferencers, { ChildId }, InNodeInfos, InCurrentDepth + 1, InMaxDepth);
+				ProvisionSize += InNodeInfos[ChildId].ProvisionSize(PluginId);
+			}
+			else if (!InNodeInfos[ChildId].Parents.Contains(PluginId))
+			{
+				InNodeInfos[ChildId].Parents.Emplace(PluginId);
+				InNodeInfos[PluginId].Children.Emplace(Pair);
+				ProvisionSize += 1;
+			}
+		}
+	}
+
+	InNodeInfos[PluginId].ChildProvisionSize = ProvisionSize > 0 ? ProvisionSize : 1;
+}
+
+UEdGraphNode_PluginReference* UEdGraph_PluginReferenceViewer::RebuildGraph()
+{
+	RemoveAllNodes();
+
+	UEdGraphNode_PluginReference* NewRootNode = ConstructNodes(CurrentGraphRootIdentifiers, CurrentGraphRootOrigin);
+	return NewRootNode;
+}
+
+UEdGraphNode_PluginReference* UEdGraph_PluginReferenceViewer::CreatePluginReferenceNode()
+{
+	const bool bSelectNewNode = false;
+	return Cast<UEdGraphNode_PluginReference>(CreateNode(UEdGraphNode_PluginReference::StaticClass(), bSelectNewNode));
+}
+
+void UEdGraph_PluginReferenceViewer::RemoveAllNodes()
+{
+	TArray<UEdGraphNode*> NodesToRemove = Nodes;
+	for (int32 NodeIndex = 0; NodeIndex < NodesToRemove.Num(); ++NodeIndex)
+	{
+		RemoveNode(NodesToRemove[NodeIndex]);
+	}
+}
+
+void UEdGraph_PluginReferenceViewer::SetPluginReferenceViewer(TSharedPtr<SPluginReferenceViewer> InViewer)
+{
+	PluginReferenceViewer = InViewer;
+}
+
+void UEdGraph_PluginReferenceViewer::GetPluginDependencies(const FPluginIdentifier& PluginIdentifier, TArray<FPluginDependency>& OutDependencies) const
+{
+	const bool bShowEnginePlugins = PluginReferenceViewer.Pin()->IsShowEnginePluginsChecked();
+	const bool bShowOptionalPlugins = PluginReferenceViewer.Pin()->IsShowOptionalPluginsChecked();
+
+	const FPluginDependsNode* PluginNode = CachedDependsNodes.FindChecked(PluginIdentifier).Get();
+	const TSharedRef<IPlugin> Plugin = PluginMap.FindChecked(PluginNode->Identifier);
+
+	for (int32 i = 0; i < PluginNode->Dependencies.Num(); ++i)
+	{
+		const FPluginDependsNode* ChildNode = PluginNode->Dependencies[i];
+
+		const TSharedRef<IPlugin> ChildPlugin = PluginMap.FindChecked(ChildNode->Identifier);
+		if (ChildPlugin->GetType() == EPluginType::Engine && !bShowEnginePlugins)
+		{
+			continue;
+		}
+
+		const FPluginReferenceDescriptor& ReferenceDescriptor = Plugin->GetDescriptor().Plugins[i];
+		check(ReferenceDescriptor.Name == ChildNode->Identifier.Name.ToString());
+
+		if (ReferenceDescriptor.bOptional && !bShowOptionalPlugins)
+		{
+			continue;
+		}
+
+		OutDependencies.Add(FPluginDependency(ChildNode->Identifier, ReferenceDescriptor.bEnabled, ReferenceDescriptor.bOptional));
+	}
+}
+
+void UEdGraph_PluginReferenceViewer::GetPluginReferencers(const FPluginIdentifier& PluginIdentifier, TArray<FPluginDependency>& OutReferencers) const
+{
+	const FPluginDependsNode* PluginNode = CachedDependsNodes.FindChecked(PluginIdentifier).Get();
+	for (const FPluginDependsNode* Referencer : PluginNode->Referencers)
+	{
+		for (int32 i = 0; i < Referencer->Dependencies.Num(); ++i)
+		{
+			const FPluginDependsNode* ParentNode = Referencer;
+
+			if (ParentNode->Dependencies[i] == PluginNode)
+			{
+				const TSharedRef<IPlugin> ParentPlugin = PluginMap.FindChecked(Referencer->Identifier);
+				const FPluginReferenceDescriptor& ReferenceDescriptor = ParentPlugin->GetDescriptor().Plugins[i];
+				check(ReferenceDescriptor.Name == PluginNode->Identifier.Name.ToString());
+
+				OutReferencers.Add(FPluginDependency(ParentNode->Identifier, ReferenceDescriptor.bEnabled, ReferenceDescriptor.bOptional));
+				break;
+			}
+		}
+	}
+}
\ No newline at end of file
diff --git a/Engine/Plugins/Experimental/PluginAudit/Source/Private/EdGraph_PluginReferenceViewer.h b/Engine/Plugins/Experimental/PluginAudit/Source/Private/EdGraph_PluginReferenceViewer.h
new file mode 100644
--- /dev/null
+++ b/Engine/Plugins/Experimental/PluginAudit/Source/Private/EdGraph_PluginReferenceViewer.h
@@ -0,0 +1,139 @@
+// Copyright Epic Games, Inc. All Rights Reserved.
+
+#pragma once
+
+#include "CoreMinimal.h"
+#include "EdGraph/EdGraph.h"
+#include "PluginReferencePinCategory.h"
+
+#include "EdGraph_PluginReferenceViewer.generated.h"
+
+class IPlugin;
+class UEdGraphNode_PluginReference;
+class SPluginReferenceViewer;
+
+struct FPluginIdentifier
+{
+	FName Name;
+
+	FPluginIdentifier()
+		: Name(NAME_None)
+	{
+	}
+
+	FPluginIdentifier(FName InName)
+		: Name(InName)
+	{
+	}
+
+	FPluginIdentifier(const FString& InName)
+		: Name(InName)
+	{
+	}
+
+	friend inline bool operator==(const FPluginIdentifier& A, const FPluginIdentifier& B)
+	{
+		return A.Name == B.Name;
+	}
+
+	friend inline uint32 GetTypeHash(const FPluginIdentifier& Key)
+	{
+		return GetTypeHash(Key.Name);
+	}
+};
+
+struct FPluginDependsNode
+{
+	explicit FPluginDependsNode(const FPluginIdentifier& InIdentifier)
+		: Identifier(InIdentifier)
+	{
+	}
+
+	FPluginIdentifier Identifier;
+	TArray<FPluginDependsNode*> Dependencies;
+	TArray<FPluginDependsNode*> Referencers;
+};
+
+struct FPluginReferenceNodeInfo
+{
+	FPluginIdentifier Identifier;
+
+	TSharedPtr<IPlugin> Plugin;
+
+	TArray<TPair<FPluginIdentifier, EPluginReferencePinCategory>> Children;
+
+	TArray<FPluginIdentifier> Parents;
+
+	bool bReferencers;
+
+	FPluginReferenceNodeInfo(const FPluginIdentifier& InIdentifier, bool bInReferencers);
+
+	bool IsFirstParent(const FPluginIdentifier& InParentId) const;
+
+	// The Provision Size, or vertical spacing required for layout, for a given parent.  
+	int32 ProvisionSize(const FPluginIdentifier& InParentId) const;
+
+	// how many nodes worth of children require vertical spacing 
+	int32 ChildProvisionSize;
+};
+
+UCLASS()
+class UEdGraph_PluginReferenceViewer : public UEdGraph
+{
+	GENERATED_UCLASS_BODY()
+
+public:
+	void CachePluginDependencies(const TArray<TSharedRef<IPlugin>>& Plugins);
+
+	/** Set reference viewer to focus on these assets */
+	void SetGraphRoot(const TArray<FPluginIdentifier>& GraphRootIdentifiers, const FIntPoint& GraphRootOrigin = FIntPoint(ForceInitToZero));
+
+	UEdGraphNode_PluginReference* ConstructNodes(const TArray<FPluginIdentifier>& GraphRootIdentifiers, const FIntPoint& GraphRootOrigin);
+
+	UEdGraphNode_PluginReference* RecursivelyCreateNodes(bool bInReferencers, const FPluginIdentifier& InPluginId, const FIntPoint& InNodeLoc, const FPluginIdentifier& InParentId, UEdGraphNode_PluginReference* InParentNode, TMap<FPluginIdentifier, FPluginReferenceNodeInfo>& InNodeInfos, int32 InCurrentDepth, int32 InMaxDepth, bool bIsRoot = false);
+
+	void RecursivelyPopulateNodeInfos(bool bInReferencers, const TArray<FPluginIdentifier>& Identifiers, TMap<FPluginIdentifier, FPluginReferenceNodeInfo>& InNodeInfos, int32 InCurrentDepth, int32 InMaxDepth);
+
+	void GetSortedLinks(const TArray<FPluginIdentifier>& Identifiers, bool bReferencers, TMap<FPluginIdentifier, EPluginReferencePinCategory>& OutLinks) const;
+
+	/** Force the graph to rebuild */
+	class UEdGraphNode_PluginReference* RebuildGraph();
+
+	UEdGraphNode_PluginReference* CreatePluginReferenceNode();
+	void RemoveAllNodes();
+
+private:
+	struct FPluginDependency
+	{
+		FPluginIdentifier Identifier;
+		bool bIsEnabled;
+		bool bIsOptional;
+
+		FPluginDependency(FPluginIdentifier InIdentifier, bool bInIsEnabled, bool bInIsOptional)
+			: Identifier(InIdentifier)
+			, bIsEnabled(bInIsEnabled)
+			, bIsOptional(bInIsOptional)
+		{
+		}
+	};
+
+	void SetPluginReferenceViewer(TSharedPtr<SPluginReferenceViewer> InViewer);
+
+	void GetPluginDependencies(const FPluginIdentifier& PluginIdentifier, TArray<FPluginDependency>& OutDependencies) const;
+	void GetPluginReferencers(const FPluginIdentifier& PluginIdentifier, TArray<FPluginDependency>& OutReferencers) const;
+
+private:
+	/** Editor for this pool */
+	TWeakPtr<SPluginReferenceViewer> PluginReferenceViewer;
+
+	TArray<FPluginIdentifier> CurrentGraphRootIdentifiers;
+	FIntPoint CurrentGraphRootOrigin;
+
+	TMap<FPluginIdentifier, TUniquePtr<FPluginDependsNode>> CachedDependsNodes;
+	TMap<FPluginIdentifier, const TSharedRef<IPlugin>> PluginMap;
+
+	TMap<FPluginIdentifier, FPluginReferenceNodeInfo> ReferencerNodeInfos;
+	TMap<FPluginIdentifier, FPluginReferenceNodeInfo> DependencyNodeInfos;
+
+	friend SPluginReferenceViewer;
+};
\ No newline at end of file
diff --git a/Engine/Plugins/Experimental/PluginAudit/Source/Private/PluginAuditModule.cpp b/Engine/Plugins/Experimental/PluginAudit/Source/Private/PluginAuditModule.cpp
--- a/Engine/Plugins/Experimental/PluginAudit/Source/Private/PluginAuditModule.cpp
+++ b/Engine/Plugins/Experimental/PluginAudit/Source/Private/PluginAuditModule.cpp
@@ -5,9 +5,25 @@
 #include "Widgets/Docking/SDockTab.h"
 #include "WorkspaceMenuStructure.h"
 #include "WorkspaceMenuStructureModule.h"
+#include "EdGraphUtilities.h"
+#include "EdGraphNode_PluginReference.h"
+#include "SPluginReferenceNode.h"
 
 #define LOCTEXT_NAMESPACE "PluginAudit"
 
+class FPluginAuditGraphPanelNodeFactory : public FGraphPanelNodeFactory
+{
+	virtual TSharedPtr<SGraphNode> CreateNode(UEdGraphNode* Node) const override
+	{
+		if (UEdGraphNode_PluginReference* DependencyNode = Cast<UEdGraphNode_PluginReference>(Node))
+		{
+			return SNew(SPluginReferenceNode, DependencyNode);
+		}
+
+		return nullptr;
+	}
+};
+
 
 class FPluginAuditModule : public IModuleInterface
 {
@@ -20,6 +36,8 @@ public:
 private:
 	static const FName PluginAuditTabName;
 	TSharedRef<SDockTab> SpawnPluginAuditTab(const FSpawnTabArgs& Args);
+
+	TSharedPtr<FPluginAuditGraphPanelNodeFactory> PluginAuditGraphPanelNodeFactory;
 };
 
 const FName FPluginAuditModule::PluginAuditTabName = TEXT("PluginAudit");
@@ -37,6 +55,9 @@ void FPluginAuditModule::StartupModule()
 		.SetGroup(WorkspaceMenu::GetMenuStructure().GetDeveloperToolsAuditCategory())
 		.SetIcon(FSlateIcon(FAppStyle::GetAppStyleSetName(), "LevelEditor.Audit"));
 	FGlobalTabmanager::Get()->RegisterDefaultTabWindowSize(PluginAuditTabName, FVector2D(1080, 600));
+
+	PluginAuditGraphPanelNodeFactory = MakeShareable(new FPluginAuditGraphPanelNodeFactory());
+	FEdGraphUtilities::RegisterVisualNodeFactory(PluginAuditGraphPanelNodeFactory);
 }
 
 void FPluginAuditModule::ShutdownModule()
diff --git a/Engine/Plugins/Experimental/PluginAudit/Source/Private/PluginReferencePinCategory.cpp b/Engine/Plugins/Experimental/PluginAudit/Source/Private/PluginReferencePinCategory.cpp
new file mode 100644
--- /dev/null
+++ b/Engine/Plugins/Experimental/PluginAudit/Source/Private/PluginReferencePinCategory.cpp
@@ -0,0 +1,89 @@
+// Copyright Epic Games, Inc. All Rights Reserved.
+
+#include "PluginReferencePinCategory.h"
+
+namespace PluginDependencyPinCategory
+{
+	FName NamePassive(TEXT("Passive"));
+	FName NameHard(TEXT("Hard"));
+	FName NameHardOptional(TEXT("HardOptional"));
+	FName NameSoft(TEXT("Soft"));
+	FName NameSoftOptional(TEXT("SoftOptional"));
+	const FLinearColor ColorPassive = FLinearColor(128, 128, 128);
+	const FLinearColor ColorHard = FLinearColor(FColor(236, 252, 227)); // RiceFlower
+	const FLinearColor ColorHardOptional = FLinearColor(FColor(118, 126, 114));
+	const FLinearColor ColorSoft = FLinearColor(FColor(145, 66, 117)); // CannonPink
+	const FLinearColor ColorSoftOptional = FLinearColor(FColor(73, 33, 58));
+}
+
+namespace PluginReferencePinUtil
+{
+	EPluginReferencePinCategory ParseDependencyPinCategory(FName PinCategory)
+	{
+		if (PinCategory == PluginDependencyPinCategory::NameHardOptional)
+		{
+			return EPluginReferencePinCategory::LinkEndActive | EPluginReferencePinCategory::LinkTypeEnabled | EPluginReferencePinCategory::LinkTypeOptional;
+		}
+		else if (PinCategory == PluginDependencyPinCategory::NameHard)
+		{
+			return EPluginReferencePinCategory::LinkEndActive | EPluginReferencePinCategory::LinkTypeEnabled;
+		}
+		else if (PinCategory == PluginDependencyPinCategory::NameSoftOptional)
+		{
+			return EPluginReferencePinCategory::LinkEndActive | EPluginReferencePinCategory::LinkTypeOptional;
+		}
+		else if (PinCategory == PluginDependencyPinCategory::NameSoft)
+		{
+			return EPluginReferencePinCategory::LinkEndActive;
+		}
+		else
+		{
+			return EPluginReferencePinCategory::LinkEndPassive;
+		}
+	}
+
+	FName GetName(EPluginReferencePinCategory Category)
+	{
+		if ((Category & EPluginReferencePinCategory::LinkEndMask) == EPluginReferencePinCategory::LinkEndPassive)
+		{
+			return PluginDependencyPinCategory::NamePassive;
+		}
+		else
+		{
+			switch (Category & EPluginReferencePinCategory::LinkTypeMask)
+			{
+			case EPluginReferencePinCategory::LinkTypeEnabled | EPluginReferencePinCategory::LinkTypeOptional:
+				return PluginDependencyPinCategory::NameHardOptional;
+			case EPluginReferencePinCategory::LinkTypeEnabled:
+				return PluginDependencyPinCategory::NameHard;
+			case EPluginReferencePinCategory::LinkTypeOptional:
+				return PluginDependencyPinCategory::NameSoftOptional;
+			default:
+				return PluginDependencyPinCategory::NameSoft;
+			}
+		}
+	}
+
+
+	FLinearColor GetColor(EPluginReferencePinCategory Category)
+	{
+		if ((Category & EPluginReferencePinCategory::LinkEndMask) == EPluginReferencePinCategory::LinkEndPassive)
+		{
+			return PluginDependencyPinCategory::ColorPassive;
+		}
+		else
+		{
+			switch (Category & EPluginReferencePinCategory::LinkTypeMask)
+			{
+			case EPluginReferencePinCategory::LinkTypeEnabled | EPluginReferencePinCategory::LinkTypeOptional:
+				return PluginDependencyPinCategory::ColorHardOptional;
+			case EPluginReferencePinCategory::LinkTypeEnabled:
+				return PluginDependencyPinCategory::ColorHard;
+			case EPluginReferencePinCategory::LinkTypeOptional:
+				return PluginDependencyPinCategory::ColorSoftOptional;
+			default:
+				return PluginDependencyPinCategory::ColorSoft;
+			}
+		}
+	}
+}
\ No newline at end of file
diff --git a/Engine/Plugins/Experimental/PluginAudit/Source/Private/PluginReferencePinCategory.h b/Engine/Plugins/Experimental/PluginAudit/Source/Private/PluginReferencePinCategory.h
new file mode 100644
--- /dev/null
+++ b/Engine/Plugins/Experimental/PluginAudit/Source/Private/PluginReferencePinCategory.h
@@ -0,0 +1,25 @@
+// Copyright Epic Games, Inc. All Rights Reserved.
+
+#pragma once
+
+#include "CoreMinimal.h"
+
+enum class EPluginReferencePinCategory
+{
+	LinkEndPassive = 0,
+	LinkEndActive = 1,
+	LinkEndMask = LinkEndActive,
+
+	LinkTypeNone = 0,
+	LinkTypeOptional = 2,
+	LinkTypeEnabled = 4,
+	LinkTypeMask = LinkTypeEnabled | LinkTypeOptional,
+};
+ENUM_CLASS_FLAGS(EPluginReferencePinCategory);
+
+namespace PluginReferencePinUtil
+{
+	EPluginReferencePinCategory ParseDependencyPinCategory(FName PinCategory);
+	FName GetName(EPluginReferencePinCategory Category);
+	FLinearColor GetColor(EPluginReferencePinCategory Category);
+}
\ No newline at end of file
diff --git a/Engine/Plugins/Experimental/PluginAudit/Source/Private/PluginReferenceViewerCommands.cpp b/Engine/Plugins/Experimental/PluginAudit/Source/Private/PluginReferenceViewerCommands.cpp
new file mode 100644
--- /dev/null
+++ b/Engine/Plugins/Experimental/PluginAudit/Source/Private/PluginReferenceViewerCommands.cpp
@@ -0,0 +1,27 @@
+// Copyright Epic Games, Inc. All Rights Reserved.
+
+#include "PluginReferenceViewerCommands.h"
+#include "Styling/AppStyle.h"
+
+#define LOCTEXT_NAMESPACE "PluginReferenceViewerCommands"
+
+FPluginReferenceViewerCommands::FPluginReferenceViewerCommands() : TCommands<FPluginReferenceViewerCommands>(
+	"PluginReferenceViewerCommands",
+	NSLOCTEXT("Contexts", "PluginReferenceViewerCommands", "Plugin Reference Viewer"),
+	NAME_None,
+	FAppStyle::GetAppStyleSetName())
+{
+}
+
+void FPluginReferenceViewerCommands::RegisterCommands()
+{	
+	UI_COMMAND(CompactMode, "Compact Mode", "Toggles Compact View", EUserInterfaceActionType::ToggleButton, FInputChord(EKeys::V));
+	UI_COMMAND(ShowEnginePlugins, "Show Engine Plugins", "Toggles visibility of Engine Plugins", EUserInterfaceActionType::ToggleButton, FInputChord());
+	UI_COMMAND(ShowOptionalPlugins, "Show Optional Plugins", "Toggles visibility of Optional Plugin dependencies", EUserInterfaceActionType::ToggleButton, FInputChord());
+
+	UI_COMMAND(OpenPluginProperties, "Open Plugin Properties", "Opens the Plugin Properties Editor Window", EUserInterfaceActionType::Button, FInputChord());
+	UI_COMMAND(ZoomToFit, "Zoom to Fit", "Zoom in and center the view on the selected item", EUserInterfaceActionType::Button, FInputChord(EKeys::F));
+	UI_COMMAND(ReCenterGraph, "Re-Center Graph", "Re-centers the graph on this node, showing all referencers and references for this Plugin instead", EUserInterfaceActionType::Button, FInputChord());
+}
+
+#undef LOCTEXT_NAMESPACE
\ No newline at end of file
diff --git a/Engine/Plugins/Experimental/PluginAudit/Source/Private/PluginReferenceViewerCommands.h b/Engine/Plugins/Experimental/PluginAudit/Source/Private/PluginReferenceViewerCommands.h
new file mode 100644
--- /dev/null
+++ b/Engine/Plugins/Experimental/PluginAudit/Source/Private/PluginReferenceViewerCommands.h
@@ -0,0 +1,24 @@
+// Copyright Epic Games, Inc. All Rights Reserved.
+
+#pragma once
+
+#include "Framework/Commands/Commands.h"
+
+// Actions that can be invoked in the reference viewer
+class FPluginReferenceViewerCommands : public TCommands<FPluginReferenceViewerCommands>
+{
+public:
+	FPluginReferenceViewerCommands();
+
+	// TCommands<> interface
+	virtual void RegisterCommands() override;
+	// End of TCommands<> interface
+
+	TSharedPtr<FUICommandInfo> CompactMode;
+	TSharedPtr<FUICommandInfo> ShowEnginePlugins;
+	TSharedPtr<FUICommandInfo> ShowOptionalPlugins;
+
+	TSharedPtr<FUICommandInfo> OpenPluginProperties;
+	TSharedPtr<FUICommandInfo> ZoomToFit;
+	TSharedPtr<FUICommandInfo> ReCenterGraph;
+};
\ No newline at end of file
diff --git a/Engine/Plugins/Experimental/PluginAudit/Source/Private/PluginReferenceViewerSchema.cpp b/Engine/Plugins/Experimental/PluginAudit/Source/Private/PluginReferenceViewerSchema.cpp
new file mode 100644
--- /dev/null
+++ b/Engine/Plugins/Experimental/PluginAudit/Source/Private/PluginReferenceViewerSchema.cpp
@@ -0,0 +1,101 @@
+// Copyright Epic Games, Inc. All Rights Reserved.
+
+#include "PluginReferenceViewerSchema.h"
+
+#include "ConnectionDrawingPolicy.h"
+#include "PluginReferencePinCategory.h"
+#include "PluginReferenceViewerCommands.h"
+#include "SPluginReferenceViewer.h"
+#include "ToolMenu.h"
+#include "ToolMenuEntry.h"
+#include "ToolMenuSection.h"
+
+#include UE_INLINE_GENERATED_CPP_BY_NAME(PluginReferenceViewerSchema)
+
+// Overridden connection drawing policy to use less curvy lines between nodes
+class FReferenceViewerConnectionDrawingPolicy : public FConnectionDrawingPolicy
+{
+public:
+	FReferenceViewerConnectionDrawingPolicy(int32 InBackLayerID, int32 InFrontLayerID, float InZoomFactor, const FSlateRect& InClippingRect, FSlateWindowElementList& InDrawElements)
+		: FConnectionDrawingPolicy(InBackLayerID, InFrontLayerID, InZoomFactor, InClippingRect, InDrawElements)
+	{
+	}
+
+	virtual FVector2D ComputeSplineTangent(const FVector2D& Start, const FVector2D& End) const override
+	{
+		const int32 Tension = FMath::Abs<int32>(Start.X - End.X);
+		return Tension * FVector2D(1.0f, 0);
+	}
+
+	virtual void DetermineWiringStyle(UEdGraphPin* OutputPin, UEdGraphPin* InputPin, /*inout*/ FConnectionParams& Params) override
+	{
+		EPluginReferencePinCategory OutputCategory = PluginReferencePinUtil::ParseDependencyPinCategory(OutputPin->PinType.PinCategory);
+		EPluginReferencePinCategory InputCategory = PluginReferencePinUtil::ParseDependencyPinCategory(InputPin->PinType.PinCategory);
+
+		EPluginReferencePinCategory Category = !!(OutputCategory & EPluginReferencePinCategory::LinkEndActive) ? OutputCategory : InputCategory;
+		Params.WireColor = PluginReferencePinUtil::GetColor(Category);
+	}
+};
+
+//////////////////////////////////////////////////////////////////////////
+// UReferenceViewerSchema
+
+UPluginReferenceViewerSchema::UPluginReferenceViewerSchema(const FObjectInitializer& ObjectInitializer)
+	: Super(ObjectInitializer)
+{
+}
+
+void UPluginReferenceViewerSchema::GetContextMenuActions(UToolMenu* Menu, UGraphNodeContextMenuContext* Context) const
+{
+	{
+		FToolMenuSection& Section = Menu->AddSection(TEXT("Plugin"), NSLOCTEXT("PluginReferenceViewerSchema", "PluginSectionLabel", "Plugin"));
+		Section.AddMenuEntry(FPluginReferenceViewerCommands::Get().OpenPluginProperties);
+	}
+
+	{
+		FToolMenuSection& Section = Menu->AddSection(TEXT("Misc"), NSLOCTEXT("PluginReferenceViewerSchema", "MiscSectionLabel", "Misc"));
+		Section.AddMenuEntry(FPluginReferenceViewerCommands::Get().ZoomToFit);
+		Section.AddMenuEntry(FPluginReferenceViewerCommands::Get().ReCenterGraph);
+	}
+}
+
+FLinearColor UPluginReferenceViewerSchema::GetPinTypeColor(const FEdGraphPinType& PinType) const
+{
+	return PluginReferencePinUtil::GetColor(PluginReferencePinUtil::ParseDependencyPinCategory(PinType.PinCategory));
+}
+
+void UPluginReferenceViewerSchema::BreakPinLinks(UEdGraphPin& TargetPin, bool bSendsNodeNotifcation) const
+{
+	// Don't allow breaking any links
+}
+
+void UPluginReferenceViewerSchema::BreakSinglePinLink(UEdGraphPin* SourcePin, UEdGraphPin* TargetPin) const
+{
+	// Don't allow breaking any links
+}
+
+FPinConnectionResponse UPluginReferenceViewerSchema::MovePinLinks(UEdGraphPin& MoveFromPin, UEdGraphPin& MoveToPin, bool bIsIntermediateMove, bool bNotifyLinkedNodes) const
+{
+	// Don't allow moving any links
+	return FPinConnectionResponse(CONNECT_RESPONSE_DISALLOW, FString());
+}
+
+FPinConnectionResponse UPluginReferenceViewerSchema::CopyPinLinks(UEdGraphPin& CopyFromPin, UEdGraphPin& CopyToPin, bool bIsIntermediateCopy) const
+{
+	// Don't allow copying any links
+	return FPinConnectionResponse(CONNECT_RESPONSE_DISALLOW, FString());
+}
+
+FConnectionDrawingPolicy* UPluginReferenceViewerSchema::CreateConnectionDrawingPolicy(int32 InBackLayerID, int32 InFrontLayerID, float InZoomFactor, const FSlateRect& InClippingRect, class FSlateWindowElementList& InDrawElements, class UEdGraph* InGraphObj) const
+{
+	return new FReferenceViewerConnectionDrawingPolicy(InBackLayerID, InFrontLayerID, InZoomFactor, InClippingRect, InDrawElements);
+}
+
+void UPluginReferenceViewerSchema::DroppedAssetsOnGraph(const TArray<FAssetData>& Assets, const FVector2D& GraphPosition, UEdGraph* Graph) const
+{
+}
+
+void UPluginReferenceViewerSchema::GetAssetsGraphHoverMessage(const TArray<FAssetData>& Assets, const UEdGraph* HoverGraph, FString& OutTooltipText, bool& OutOkIcon) const
+{
+	OutOkIcon = true;
+}
\ No newline at end of file
diff --git a/Engine/Plugins/Experimental/PluginAudit/Source/Private/PluginReferenceViewerSchema.h b/Engine/Plugins/Experimental/PluginAudit/Source/Private/PluginReferenceViewerSchema.h
new file mode 100644
--- /dev/null
+++ b/Engine/Plugins/Experimental/PluginAudit/Source/Private/PluginReferenceViewerSchema.h
@@ -0,0 +1,32 @@
+// Copyright Epic Games, Inc. All Rights Reserved.
+
+#pragma once
+
+#include "CoreMinimal.h"
+#include "UObject/ObjectMacros.h"
+#include "EdGraph/EdGraphSchema.h"
+#include "PluginReferenceViewerSchema.generated.h"
+
+class UToolMenu;
+class FSlateRect;
+class UEdGraph;
+
+UCLASS(MinimalAPI)
+class UPluginReferenceViewerSchema : public UEdGraphSchema
+{
+	GENERATED_UCLASS_BODY()
+
+public:
+	// UEdGraphSchema interface
+	virtual void GetContextMenuActions(class UToolMenu* Menu, class UGraphNodeContextMenuContext* Context) const override;
+	virtual FName GetParentContextMenuName() const override { return NAME_None; }
+	virtual FLinearColor GetPinTypeColor(const FEdGraphPinType& PinType) const override;
+	virtual void BreakPinLinks(UEdGraphPin& TargetPin, bool bSendsNodeNotifcation) const override;
+	virtual void BreakSinglePinLink(UEdGraphPin* SourcePin, UEdGraphPin* TargetPin) const override;
+	virtual FPinConnectionResponse MovePinLinks(UEdGraphPin& MoveFromPin, UEdGraphPin& MoveToPin, bool bIsIntermediateMove = false, bool bNotifyLinkedNodes = false) const override;
+	virtual FPinConnectionResponse CopyPinLinks(UEdGraphPin& CopyFromPin, UEdGraphPin& CopyToPin, bool bIsIntermediateCopy = false) const override;
+	virtual class FConnectionDrawingPolicy* CreateConnectionDrawingPolicy(int32 InBackLayerID, int32 InFrontLayerID, float InZoomFactor, const FSlateRect& InClippingRect, class FSlateWindowElementList& InDrawElements, class UEdGraph* InGraphObj) const override;
+	virtual void DroppedAssetsOnGraph(const TArray<FAssetData>& Assets, const FVector2D& GraphPosition, UEdGraph* Graph) const override;
+	virtual void GetAssetsGraphHoverMessage(const TArray<FAssetData>& Assets, const UEdGraph* HoverGraph, FString& OutTooltipText, bool& OutOkIcon) const override;
+	// End of UEdGraphSchema interface
+};
\ No newline at end of file
diff --git a/Engine/Plugins/Experimental/PluginAudit/Source/Private/PluginReferenceViewerStyle.cpp b/Engine/Plugins/Experimental/PluginAudit/Source/Private/PluginReferenceViewerStyle.cpp
new file mode 100644
--- /dev/null
+++ b/Engine/Plugins/Experimental/PluginAudit/Source/Private/PluginReferenceViewerStyle.cpp
@@ -0,0 +1,106 @@
+// Copyright Epic Games, Inc. All Rights Reserved.
+
+#include "PluginReferenceViewerStyle.h"
+
+#include "Brushes/SlateImageBrush.h"
+#include "Styling/SlateTypes.h"
+#include "Brushes/SlateRoundedBoxBrush.h"
+#include "Styling/SlateStyleMacros.h"
+#include "Misc/Paths.h"
+#include "Styling/StyleColors.h"
+#include "Styling/CoreStyle.h"
+#include "Styling/SlateStyleRegistry.h"
+
+
+FName FPluginReferenceViewerStyle::StyleName("PluginReferenceViewerStyle");
+
+#define FONT(...) FSlateFontInfo(FCoreStyle::GetDefaultFont(), __VA_ARGS__)
+
+FPluginReferenceViewerStyle::FPluginReferenceViewerStyle()
+	: FSlateStyleSet(StyleName)
+{
+	SetParentStyleName("EditorStyle");
+
+	SetContentRoot(FPaths::EnginePluginsDir() / TEXT("Editor/AssetManagerEditor/Content"));
+	SetCoreContentRoot(FPaths::EngineContentDir() / TEXT("Slate"));
+
+	FScrollBarStyle ScrollBar = GetParentStyle()->GetWidgetStyle<FScrollBarStyle>("ScrollBar");
+	FTextBlockStyle NormalText = GetParentStyle()->GetWidgetStyle<FTextBlockStyle>("NormalText");
+	FEditableTextBoxStyle NormalEditableTextBoxStyle = GetParentStyle()->GetWidgetStyle<FEditableTextBoxStyle>("NormalEditableTextBox");
+
+	FSlateFontInfo TitleFont = FONT(12, "Bold");
+
+	// Text Styles 
+	FTextBlockStyle GraphNodeTitle = FTextBlockStyle(NormalText)
+		.SetFont(TitleFont)
+		.SetColorAndOpacity( FStyleColors::White)
+		.SetShadowOffset( FVector2D::UnitVector )
+		.SetShadowColorAndOpacity( FLinearColor::Black );
+	Set( "Graph.Node.NodeTitle", GraphNodeTitle );
+
+
+	Set( "Graph.Node.NodeTitleExtraLines", FTextBlockStyle(NormalText)
+		.SetFont( DEFAULT_FONT( "Normal", 9 ) )
+		.SetColorAndOpacity( FStyleColors::White)
+		.SetShadowOffset( FVector2D::ZeroVector)
+		.SetShadowColorAndOpacity( FLinearColor::Transparent)
+	);
+
+
+	FEditableTextBoxStyle GraphNodeTitleEditableText = FEditableTextBoxStyle(NormalEditableTextBoxStyle)
+		.SetFont(TitleFont)
+		.SetForegroundColor(FStyleColors::Input)
+		.SetBackgroundImageNormal(FSlateRoundedBoxBrush(FStyleColors::Foreground, FStyleColors::Secondary, 1.0f))
+		.SetBackgroundImageHovered(FSlateRoundedBoxBrush(FStyleColors::Foreground, FStyleColors::Hover, 1.0f))
+		.SetBackgroundImageFocused(FSlateRoundedBoxBrush(FStyleColors::Foreground, FStyleColors::Primary, 1.0f))
+		.SetBackgroundImageReadOnly(FSlateRoundedBoxBrush(FStyleColors::Header, FStyleColors::InputOutline, 1.0f))
+		.SetForegroundColor(FStyleColors::White)
+		.SetBackgroundColor(FStyleColors::White)
+		.SetReadOnlyForegroundColor(FStyleColors::Foreground)
+		.SetFocusedForegroundColor(FStyleColors::Background)
+		.SetScrollBarStyle( ScrollBar );
+	Set( "Graph.Node.NodeTitleEditableText", GraphNodeTitleEditableText );
+
+	Set( "Graph.Node.NodeTitleInlineEditableText", FInlineEditableTextBlockStyle()
+		.SetTextStyle(GraphNodeTitle)
+		.SetEditableTextBoxStyle(GraphNodeTitleEditableText)
+	);
+
+	
+	FLinearColor SpillColor(.3, .3, .3, 1.0);
+	int BodyRadius = 10.0; // Designed for 4 but using 10 to accomodate the shared selection border.  Update to 4 all nodes get aligned.
+
+	// NOTE: 
+	Set( "Graph.Node.BodyBackground", new FSlateRoundedBoxBrush(FStyleColors::Panel, BodyRadius));
+	Set( "Graph.Node.BodyBorder",     new FSlateRoundedBoxBrush(SpillColor, BodyRadius));
+
+	Set( "Graph.Node.Body",           new FSlateRoundedBoxBrush(FStyleColors::Panel, BodyRadius, FStyleColors::Transparent, 2.0));
+	Set( "Graph.Node.ColorSpill",     new FSlateRoundedBoxBrush(SpillColor, FVector4(BodyRadius, BodyRadius, 0.0, 0.0)));
+
+	Set( "Graph.Node.Duplicate",      new IMAGE_BRUSH_SVG("/GraphNode_Duplicate_8px", FVector2D(8.f, 8.f), FStyleColors::White));
+
+
+	const FVector2D IconSize(20.0f, 20.0f);
+
+	Set("Icons.ArrowLeft", new CORE_IMAGE_BRUSH_SVG("Starship/Common/arrow-left", IconSize));
+	Set("Icons.ArrowRight", new CORE_IMAGE_BRUSH_SVG("Starship/Common/arrow-right", IconSize));
+
+	Set("Icons.AutoFilters", new CORE_IMAGE_BRUSH_SVG("Starship/Common/FilterAuto", IconSize));
+	Set("Icons.Filters", new CORE_IMAGE_BRUSH_SVG("Starship/Common/filter", IconSize));
+	Set("Icons.Duplicate",      new IMAGE_BRUSH_SVG("/GraphNode_Duplicate_8px", IconSize, FStyleColors::White));
+
+	FSlateStyleRegistry::RegisterSlateStyle(*this);
+}
+
+FPluginReferenceViewerStyle::~FPluginReferenceViewerStyle()
+{
+	FSlateStyleRegistry::UnRegisterSlateStyle(*this);
+}
+
+FPluginReferenceViewerStyle& FPluginReferenceViewerStyle::Get()
+{
+	static FPluginReferenceViewerStyle Inst;
+	return Inst;
+}
+
+
diff --git a/Engine/Plugins/Experimental/PluginAudit/Source/Private/PluginReferenceViewerStyle.h b/Engine/Plugins/Experimental/PluginAudit/Source/Private/PluginReferenceViewerStyle.h
new file mode 100644
--- /dev/null
+++ b/Engine/Plugins/Experimental/PluginAudit/Source/Private/PluginReferenceViewerStyle.h
@@ -0,0 +1,20 @@
+// Copyright Epic Games, Inc. All Rights Reserved.
+
+#pragma once
+
+#include "Styling/SlateStyle.h"
+
+class FPluginReferenceViewerStyle
+	: public FSlateStyleSet
+{
+public:
+	static FName StyleName;
+
+	/** Access the singleton instance for this style set */
+	static FPluginReferenceViewerStyle& Get();
+
+private:
+
+	FPluginReferenceViewerStyle();
+	~FPluginReferenceViewerStyle();
+};
diff --git a/Engine/Plugins/Experimental/PluginAudit/Source/Private/SPluginAuditBrowser.cpp b/Engine/Plugins/Experimental/PluginAudit/Source/Private/SPluginAuditBrowser.cpp
--- a/Engine/Plugins/Experimental/PluginAudit/Source/Private/SPluginAuditBrowser.cpp
+++ b/Engine/Plugins/Experimental/PluginAudit/Source/Private/SPluginAuditBrowser.cpp
@@ -20,6 +20,7 @@
 #include "IMessageLogListing.h"
 #include "MessageLogModule.h"
 #include "Misc/UObjectToken.h"
+#include "SPluginReferenceViewer.h"
 #include "Widgets/Input/SSearchBox.h"
 #include "Widgets/SBoxPanel.h"
 #include "Widgets/Layout/SBorder.h"
@@ -58,7 +59,8 @@ void SPluginAuditBrowser::Construct(const FArguments& InArgs)
 	PluginListView =
 		SNew(SListView< TSharedRef<FCookedPlugin> >)
 		.ListItemsSource(&FilteredCookedPlugins)
-		.OnGenerateRow(this, &SPluginAuditBrowser::MakeCookedPluginRow);
+		.OnGenerateRow(this, &SPluginAuditBrowser::MakeCookedPluginRow)
+		.OnContextMenuOpening(this, &SPluginAuditBrowser::OnContextMenuOpening);
 
 	ChildSlot
 	[
@@ -112,12 +114,41 @@ void SPluginAuditBrowser::Construct(const FArguments& InArgs)
 					[
 						PluginListView.ToSharedRef()
 					]
+
+					+ SVerticalBox::Slot()
+					.AutoHeight()
+					[
+						SNew(SHorizontalBox)
+
+							// Asset count
+						+ SHorizontalBox::Slot()
+						.FillWidth(1.f)
+						.VAlign(VAlign_Center)
+						.Padding(8, 5)
+						[
+							SNew(STextBlock)
+							.Text(this, &SPluginAuditBrowser::GetPluginCountText)
+						]
+					]
 				]
 			]
 			+ SSplitter::Slot()
-			.Value(0.70f)
+			.Value(0.90f)
 			[
-				MessageLogModule.CreateLogListingWidget(LogListing.ToSharedRef())
+				SNew(SVerticalBox)
+
+				// Path and history
+				+SVerticalBox::Slot()
+				.FillHeight(0.90f)
+				[
+					SAssignNew(PluginReferenceViewer, SPluginReferenceViewer)
+				]
+
+				+ SVerticalBox::Slot()
+				.AutoHeight()
+				[
+					MessageLogModule.CreateLogListingWidget(LogListing.ToSharedRef())
+				]
 			]
 		]
 	];
@@ -232,6 +263,24 @@ ECheckBoxState SPluginAuditBrowser::GetGlobalDisabledState() const
 		bGlobalDisabledState ? ECheckBoxState::Unchecked : ECheckBoxState::Checked;
 }
 
+FText SPluginAuditBrowser::GetPluginCountText() const
+{
+	const int32 NumPlugins = CookedPlugins.Num();
+	const int32 NumPluginsDisabled = CookedPlugins.FilterByPredicate([](const auto& Item) { return Item->bSimulateDisabled; }).Num();
+
+	FText PluginCount = FText::GetEmpty();
+	if (NumPlugins == 1)
+	{
+		PluginCount = FText::Format(LOCTEXT("PluginCountLabelSingularPlusDisabled", "1 plugin ({0} disabled)"), FText::AsNumber(NumPluginsDisabled));
+	}
+	else
+	{
+		PluginCount = FText::Format(LOCTEXT("PluginCountLabelPluralPlusDisabled", "{0} plugins ({1} disabled)"), FText::AsNumber(CookedPlugins.Num()), FText::AsNumber(NumPluginsDisabled));
+	}
+
+	return PluginCount;
+}
+
 TSharedRef<ITableRow> SPluginAuditBrowser::MakeCookedPluginRow(TSharedRef<SPluginAuditBrowser::FCookedPlugin> InItem, const TSharedRef<STableViewBase>& OwnerTable)
 {
 	return SNew(STableRow< TSharedRef<FCookedPlugin> >, OwnerTable)
@@ -262,6 +311,65 @@ TSharedRef<ITableRow> SPluginAuditBrowser::MakeCookedPluginRow(TSharedRef<SPlugi
 	];
 }
 
+TSharedPtr<SWidget> SPluginAuditBrowser::OnContextMenuOpening()
+{
+	const bool bShouldCloseWindowAfterMenuSelection = true;
+	FMenuBuilder MenuBuilder(bShouldCloseWindowAfterMenuSelection, nullptr);
+
+	MenuBuilder.AddMenuEntry(
+		LOCTEXT("OpenPluginPropertiesLabel", "Open Plugin Properties"),
+		LOCTEXT("OpenPluginPropertiesTooltip", "Open Plugin Properties Window"),
+		FSlateIcon(),
+		FUIAction(
+			FExecuteAction::CreateSP(this, &SPluginAuditBrowser::OnOpenPluginProperties),
+			FCanExecuteAction()
+		));
+
+	MenuBuilder.AddMenuEntry(
+		LOCTEXT("OpenPluginDependencyViewerLabel", "Open Plugin Dependency Viewer"),
+		LOCTEXT("OpenPluginDependencyViewerTooltip", "Open Plugin Dependency Viewer"),
+		FSlateIcon(),
+		FUIAction(
+			FExecuteAction::CreateSP(this, &SPluginAuditBrowser::OnOpenPluginDependencyViewer),
+			FCanExecuteAction()
+		));
+
+	return MenuBuilder.MakeWidget();
+}
+
+void SPluginAuditBrowser::OnOpenPluginProperties()
+{
+	TArray<TSharedRef<FCookedPlugin>> SelectedItems = PluginListView->GetSelectedItems();
+	if (SelectedItems.Num() == 1)
+	{
+		OpenPluginProperties(SelectedItems[0]->Plugin->GetName());
+	}
+}
+
+void SPluginAuditBrowser::OnOpenPluginDependencyViewer()
+{
+	TArray<TSharedRef<FCookedPlugin>> SelectedItems = PluginListView->GetSelectedItems();
+	if (SelectedItems.Num() == 1)
+	{
+		TSharedRef<IPlugin> SelectedPlugin = SelectedItems[0]->Plugin;
+
+		TArray<FPluginIdentifier> NewGraphRootNames;
+		NewGraphRootNames.Add(FName(SelectedPlugin->GetName()));
+
+		PluginReferenceViewer->SetGraphRoot(NewGraphRootNames);
+	}
+}
+
+void SPluginAuditBrowser::OpenPluginProperties(const FString& PluginName)
+{
+	TSharedPtr<IPlugin> Plugin = IPluginManager::Get().FindPlugin(PluginName);
+	if (Plugin != nullptr)
+	{
+		IPluginsEditorFeature& PluginEditor = IModularFeatures::Get().GetModularFeature<IPluginsEditorFeature>(EditorFeatures::PluginsEditor);
+		PluginEditor.OpenPluginEditor(Plugin.ToSharedRef(), nullptr, FSimpleDelegate());
+	}
+}
+
 void SPluginAuditBrowser::CreateLogListing()
 {
 	FMessageLogModule& MessageLogModule = FModuleManager::LoadModuleChecked<FMessageLogModule>("MessageLog");
@@ -716,4 +824,4 @@ void SPluginAuditBrowser::RefreshViolations()
 	}
 }
 
-#undef LOCTEXT_NAMESPACE
+#undef LOCTEXT_NAMESPACE
\ No newline at end of file
diff --git a/Engine/Plugins/Experimental/PluginAudit/Source/Private/SPluginAuditBrowser.h b/Engine/Plugins/Experimental/PluginAudit/Source/Private/SPluginAuditBrowser.h
--- a/Engine/Plugins/Experimental/PluginAudit/Source/Private/SPluginAuditBrowser.h
+++ b/Engine/Plugins/Experimental/PluginAudit/Source/Private/SPluginAuditBrowser.h
@@ -19,6 +19,7 @@ struct FGameplayTag;
 class IMessageLogListing;
 class IMessageToken;
 class SFilterSearchBox;
+class SPluginReferenceViewer;
 class STableViewBase;
 class ITableRow;
 class FTokenizedMessage;
@@ -44,6 +45,12 @@ private:
 	ECheckBoxState GetGlobalDisabledState() const;
 	void OnGlobalDisabledStateChanged(ECheckBoxState State);
 
+	FText GetPluginCountText() const;
+	TSharedPtr<SWidget> OnContextMenuOpening();
+	void OnOpenPluginProperties();
+	void OnOpenPluginDependencyViewer();
+	void OpenPluginProperties(const FString& PluginName);
+
 	class FCookedPlugin
 	{
 	public:
@@ -98,8 +105,10 @@ private:
 	/** The plugin search box widget */
 	TSharedPtr<class SSearchBox> SearchBoxPtr;
 
+	TSharedPtr<SPluginReferenceViewer> PluginReferenceViewer;
+
 	/** Text filter object for typing in filter text to the search box */
 	TSharedPtr<FCookedPluginTextFilter> PluginTextFilter;
 
 	bool bGlobalDisabledState = false;
-};
+};
\ No newline at end of file
diff --git a/Engine/Plugins/Experimental/PluginAudit/Source/Private/SPluginReferenceNode.cpp b/Engine/Plugins/Experimental/PluginAudit/Source/Private/SPluginReferenceNode.cpp
new file mode 100644
--- /dev/null
+++ b/Engine/Plugins/Experimental/PluginAudit/Source/Private/SPluginReferenceNode.cpp
@@ -0,0 +1,278 @@
+// Copyright Epic Games, Inc. All Rights Reserved.
+
+#include "SPluginReferenceNode.h"
+
+#include "EdGraph_PluginReferenceViewer.h"
+#include "EdGraphNode_PluginReference.h"
+#include "Framework/Application/SlateApplication.h"
+#include "HAL/PlatformFileManager.h"
+#include "Interfaces/IPluginManager.h"
+#include "IDesktopPlatform.h"
+#include "PluginReferenceViewerStyle.h"
+#include "SCommentBubble.h"
+#include "Widgets/SBoxPanel.h"
+#include "Widgets/Images/SImage.h"
+#include "Widgets/Layout/SBox.h"
+#include "Widgets/Notifications/SErrorText.h"
+#include "Widgets/Text/SInlineEditableTextBlock.h"
+
+#define LOCTEXT_NAMESPACE "PluginReferenceViewer"
+
+void SPluginReferenceNode::Construct(const FArguments& InArgs, UEdGraphNode_PluginReference* InNode)
+{
+	const int32 ThumbnailSize = 128;
+	if (InNode->AllowsThumbnail())
+	{
+		const TSharedPtr<const IPlugin> Plugin = InNode->GetPlugin();
+
+		// Plugin thumbnail image
+		FString Icon128FilePath = Plugin->GetBaseDir() / TEXT("Resources/Icon128.png");
+		if (!FPlatformFileManager::Get().GetPlatformFile().FileExists(*Icon128FilePath))
+		{
+			Icon128FilePath = IPluginManager::Get().FindPlugin(TEXT("PluginBrowser"))->GetBaseDir() / TEXT("Resources/DefaultIcon128.png");
+		}
+
+		const FName BrushName(*Icon128FilePath);
+		const FIntPoint Size = FSlateApplication::Get().GetRenderer()->GenerateDynamicImageResource(BrushName);
+		if ((Size.X > 0) && (Size.Y > 0))
+		{
+			PluginIconDynamicImageBrush = MakeShareable(new FSlateDynamicImageBrush(BrushName, FVector2D(ThumbnailSize, ThumbnailSize)));
+		}
+	}
+
+	GraphNode = InNode;
+	SetCursor(EMouseCursor::CardinalCross);
+	UpdateGraphNode();
+}
+
+// UpdateGraphNode is similar to the base, but adds the option to hide the thumbnail */
+void SPluginReferenceNode::UpdateGraphNode()
+{
+	OutputPins.Empty();
+
+	// Reset variables that are going to be exposed, in case we are refreshing an already setup node.
+	RightNodeBox.Reset();
+	LeftNodeBox.Reset();
+
+	UpdateErrorInfo();
+
+	//
+	//             ______________________
+	//            |      TITLE AREA      |
+	//            +-------+------+-------+
+	//            | (>) L |      | R (>) |
+	//            | (>) E |      | I (>) |
+	//            | (>) F |      | G (>) |
+	//            | (>) T |      | H (>) |
+	//            |       |      | T (>) |
+	//            |_______|______|_______|
+	//
+	TSharedPtr<SVerticalBox> MainVerticalBox;
+	TSharedPtr<SErrorText> ErrorText;
+	TSharedPtr<SNodeTitle> NodeTitle = SNew(SNodeTitle, GraphNode).StyleSet(&FPluginReferenceViewerStyle::Get());
+
+	// Get node icon
+	IconColor = FLinearColor::White;
+	const FSlateBrush* IconBrush = nullptr;
+	if (GraphNode != NULL && GraphNode->ShowPaletteIconOnNode())
+	{
+		IconBrush = GraphNode->GetIconAndTint(IconColor).GetOptionalIcon();
+	}
+
+	TSharedRef<SWidget> ThumbnailWidget = SNullWidget::NullWidget;
+	UEdGraphNode_PluginReference* RefGraphNode = CastChecked<UEdGraphNode_PluginReference>(GraphNode);
+	bool bIsADuplicate = /*RefGraphNode->IsADuplicate()*/false;
+
+	FLinearColor OpacityColor = /*RefGraphNode->GetIsFiltered() ? FLinearColor(1.0, 1.0, 1.0, 0.4) :*/FLinearColor::White;
+	
+	if (PluginIconDynamicImageBrush.IsValid())
+	{
+		ThumbnailWidget = SNew(SBox)
+			.WidthOverride(PluginIconDynamicImageBrush->GetImageSize().X)
+			.HeightOverride(PluginIconDynamicImageBrush->GetImageSize().Y)
+			[
+				SNew(SImage)
+				.Image(PluginIconDynamicImageBrush.Get())
+			];
+	}
+
+	ContentScale.Bind(this, &SPluginReferenceNode::GetContentScale);
+	GetOrAddSlot(ENodeZone::Center)
+	.HAlign(HAlign_Center)
+	.VAlign(VAlign_Center)
+	[
+		SAssignNew(MainVerticalBox, SVerticalBox)
+		+ SVerticalBox::Slot()
+		.AutoHeight()
+		[
+			SNew(SBorder)
+			.ColorAndOpacity(OpacityColor)
+			.BorderImage(FPluginReferenceViewerStyle::Get().GetBrush("Graph.Node.BodyBackground"))
+			.Padding(0)
+			[
+				SNew(SBorder)
+				.BorderImage(FPluginReferenceViewerStyle::Get().GetBrush("Graph.Node.BodyBorder"))
+				.BorderBackgroundColor(this, &SPluginReferenceNode::GetNodeTitleColor)
+				.Padding(0)
+				[
+					SNew(SBorder)
+					.BorderImage(FPluginReferenceViewerStyle::Get().GetBrush("Graph.Node.Body"))
+					.Padding(0)
+					[
+						SNew(SVerticalBox)
+						.ToolTipText(this, &SPluginReferenceNode::GetNodeTooltip)
+
+						+ SVerticalBox::Slot()
+						.AutoHeight()
+						.HAlign(HAlign_Fill)
+						.VAlign(VAlign_Top)
+						.Padding(0)
+						[
+							SNew(SBorder)
+							.BorderImage(FPluginReferenceViewerStyle::Get().GetBrush("Graph.Node.ColorSpill"))
+							.Padding(FMargin(10.0f, 4.0f, 6.0f, 4.0f))
+							.BorderBackgroundColor(this, &SPluginReferenceNode::GetNodeTitleColor)
+							[
+								SNew(SHorizontalBox)
+								+ SHorizontalBox::Slot()
+								.VAlign(VAlign_Center)
+								.Padding(FMargin(0.f, 0.f, 8.f, 0.f))
+								.AutoWidth()
+								[
+									SNew(SImage)
+									.Image(IconBrush)
+									.DesiredSizeOverride(FVector2D(24.0, 24.0))
+									.ColorAndOpacity(this, &SGraphNode::GetNodeTitleIconColor)
+								]
+		
+								+ SHorizontalBox::Slot()
+								.FillWidth(1.0)
+								.VAlign(VAlign_Center)
+								[
+									SNew(SVerticalBox)
+									+ SVerticalBox::Slot()
+									.AutoHeight()
+									.Padding(FMargin(0.f))
+									.VAlign(VAlign_Center)
+									[
+										SAssignNew(InlineEditableText, SInlineEditableTextBlock)
+										.Style(FPluginReferenceViewerStyle::Get(), "Graph.Node.NodeTitleInlineEditableText")
+										.Text(NodeTitle.Get(), &SNodeTitle::GetHeadTitle)
+										.OnVerifyTextChanged(this, &SPluginReferenceNode::OnVerifyNameTextChanged)
+										.OnTextCommitted(this, &SPluginReferenceNode::OnNameTextCommited)
+										.IsReadOnly(this, &SPluginReferenceNode::IsNameReadOnly)
+										.IsSelected(this, &SPluginReferenceNode::IsSelectedExclusively)
+									]
+									+ SVerticalBox::Slot()
+									.AutoHeight()
+									.Padding(FMargin(0.f))
+									[
+										NodeTitle.ToSharedRef()
+									]
+								]
+
+								+ SHorizontalBox::Slot()
+								.AutoWidth()
+								.HAlign(HAlign_Right)
+								.VAlign(VAlign_Bottom)
+								.Padding(FMargin(4.f, 0.f, 0.f, 1.f))
+								[
+									SNew(SImage)
+									.Visibility(bIsADuplicate ? EVisibility::Visible : EVisibility::Hidden)
+									.ToolTipText(LOCTEXT("DuplicateAsset", "This asset is referenced multiple times. Only the first occurance shows its decendants."))
+									.DesiredSizeOverride(FVector2D(12.0, 12.0))
+									.Image(FPluginReferenceViewerStyle::Get().GetBrush("Graph.Node.Duplicate"))
+									.ColorAndOpacity(FAppStyle::Get().GetColor("Colors.Foreground"))
+								]
+							]
+						]
+						+ SVerticalBox::Slot()
+						.AutoHeight()
+						.Padding(1.0f)
+						[
+							// POPUP ERROR MESSAGE
+							SAssignNew(ErrorText, SErrorText)
+							.BackgroundColor(this, &SPluginReferenceNode::GetErrorColor)
+							.ToolTipText(this, &SPluginReferenceNode::GetErrorMsgToolTip)
+						]
+						+ SVerticalBox::Slot()
+						.AutoHeight()
+						.HAlign(HAlign_Fill)
+						.VAlign(VAlign_Top)
+						[
+							// NODE CONTENT AREA
+							SNew(SBorder)
+							.BorderImage(FAppStyle::GetBrush("NoBorder"))
+							.HAlign(HAlign_Fill)
+							.VAlign(VAlign_Fill)
+							.Padding(FMargin(0, 3))
+							[
+								SNew(SHorizontalBox)
+								+ SHorizontalBox::Slot()
+								.AutoWidth()
+								.VAlign(VAlign_Center)
+								[
+									// LEFT
+									SNew(SBox)
+									.WidthOverride(40)
+									[
+										SAssignNew(LeftNodeBox, SVerticalBox)
+									]
+								]
+
+								+ SHorizontalBox::Slot()
+								.VAlign(VAlign_Center)
+								.HAlign(HAlign_Center)
+								.FillWidth(1.0f)
+								[
+									ThumbnailWidget
+								]
+
+								+ SHorizontalBox::Slot()
+								.AutoWidth()
+								.VAlign(VAlign_Center)
+								[
+									// RIGHT
+									SNew(SBox)
+									.WidthOverride(40)
+									[
+										SAssignNew(RightNodeBox, SVerticalBox)
+									]
+								]
+							]
+						]
+
+					] // Inner Background 
+				] // Outline Border
+			] // Background 
+		]
+	];
+
+	// Create comment bubble if comment text is valid
+	GetNodeObj()->bCommentBubbleVisible = !GetNodeObj()->NodeComment.IsEmpty();
+	if (GetNodeObj()->ShouldMakeCommentBubbleVisible() && GetNodeObj()->bCommentBubbleVisible)
+	{
+		TSharedPtr<SCommentBubble> CommentBubble;
+
+		SAssignNew(CommentBubble, SCommentBubble)
+			.GraphNode(GraphNode)
+			.Text(this, &SGraphNode::GetNodeComment);
+
+		GetOrAddSlot(ENodeZone::TopCenter)
+			.SlotOffset(TAttribute<FVector2D>(CommentBubble.Get(), &SCommentBubble::GetOffset))
+			.SlotSize(TAttribute<FVector2D>(CommentBubble.Get(), &SCommentBubble::GetSize))
+			.AllowScaling(TAttribute<bool>(CommentBubble.Get(), &SCommentBubble::IsScalingAllowed))
+			.VAlign(VAlign_Top)
+			[
+				CommentBubble.ToSharedRef()
+			];
+	}
+
+	ErrorReporting = ErrorText;
+	ErrorReporting->SetError(ErrorMsg);
+	CreateBelowWidgetControls(MainVerticalBox);
+
+	CreatePinWidgets();
+}
+
+#undef LOCTEXT_NAMESPACE
diff --git a/Engine/Plugins/Experimental/PluginAudit/Source/Private/SPluginReferenceNode.h b/Engine/Plugins/Experimental/PluginAudit/Source/Private/SPluginReferenceNode.h
new file mode 100644
--- /dev/null
+++ b/Engine/Plugins/Experimental/PluginAudit/Source/Private/SPluginReferenceNode.h
@@ -0,0 +1,34 @@
+// Copyright Epic Games, Inc. All Rights Reserved.
+
+#pragma once
+
+#include "SGraphNode.h"
+
+class FAssetThumbnail;
+class UEdGraphNode_PluginReference;
+struct FSlateDynamicImageBrush;
+
+/**
+ *
+ */
+class SPluginReferenceNode : public SGraphNode
+{
+public:
+	SLATE_BEGIN_ARGS(SPluginReferenceNode) {}
+
+	SLATE_END_ARGS()
+
+	/** Constructs this widget with InArgs */
+	void Construct(const FArguments& InArgs, UEdGraphNode_PluginReference* InNode);
+
+	// SGraphNode implementation
+	virtual void UpdateGraphNode() override;
+	virtual bool IsNodeEditable() const override { return false; }
+	// End SGraphNode implementation
+
+private:
+	TSharedPtr<class FAssetThumbnail> AssetThumbnail;
+
+	/** Brush resource for the image that is dynamically loaded */
+	TSharedPtr<FSlateDynamicImageBrush> PluginIconDynamicImageBrush;
+};
\ No newline at end of file
diff --git a/Engine/Plugins/Experimental/PluginAudit/Source/Private/SPluginReferenceViewer.cpp b/Engine/Plugins/Experimental/PluginAudit/Source/Private/SPluginReferenceViewer.cpp
new file mode 100644
--- /dev/null
+++ b/Engine/Plugins/Experimental/PluginAudit/Source/Private/SPluginReferenceViewer.cpp
@@ -0,0 +1,419 @@
+// Copyright Epic Games, Inc. All Rights Reserved.
+
+#include "SPluginReferenceViewer.h"
+
+#include "EdGraphNode_PluginReference.h"
+#include "EdGraph_PluginReferenceViewer.h"
+#include "Features/EditorFeatures.h"
+#include "Features/IModularFeatures.h"
+#include "Features/IPluginsEditorFeature.h"
+#include "Interfaces/IPluginManager.h"
+#include "PluginReferenceViewerSchema.h"
+#include "PluginReferenceViewerCommands.h"
+#include "ToolMenus.h"
+#include "Widgets/Input/SSpinBox.h"
+#include "Widgets/Layout/SBorder.h"
+
+#define LOCTEXT_NAMESPACE "PluginReferenceViewer"
+
+SPluginReferenceViewer::~SPluginReferenceViewer()
+{
+	if (!GExitPurge)
+	{
+		if (ensure(GraphObj))
+		{
+			GraphObj->RemoveFromRoot();
+		}
+	}
+}
+
+void SPluginReferenceViewer::Construct(const FArguments& InArgs)
+{
+	/** Used to delay graph rebuilding during spinbox slider interaction */
+	bNeedsGraphRebuild = false;
+
+	RegisterActions();
+
+	// Create the graph
+	GraphObj = NewObject<UEdGraph_PluginReferenceViewer>();
+	GraphObj->Schema = UPluginReferenceViewerSchema::StaticClass();
+	GraphObj->AddToRoot();
+	GraphObj->SetPluginReferenceViewer(StaticCastSharedRef<SPluginReferenceViewer>(AsShared()));
+
+	SGraphEditor::FGraphEditorEvents GraphEvents;
+	GraphEvents.OnNodeDoubleClicked = FSingleNodeEvent::CreateSP(this, &SPluginReferenceViewer::OnNodeDoubleClicked);
+
+	// Create the graph editor
+	GraphEditorPtr = SNew(SGraphEditor)
+		.AdditionalCommands(PluginReferenceViewerActions)
+		.GraphToEdit(GraphObj)
+		.GraphEvents(GraphEvents)
+		.ShowGraphStateOverlay(false);
+
+	GraphObj->CachePluginDependencies(IPluginManager::Get().GetDiscoveredPlugins());
+
+	ChildSlot
+	[
+		SNew(SVerticalBox)
+
+		+SVerticalBox::Slot()
+		.AutoHeight()
+		[
+			SNew(SBorder)
+			.BorderImage(FAppStyle::GetBrush("Brushes.Panel"))
+			[
+				SNew(SHorizontalBox)
+
+				+SHorizontalBox::Slot()
+				.AutoWidth()
+				.Padding(4, 0)
+				[
+					MakeToolBar()
+				]
+			]
+		]
+
+		// Graph
+		+SVerticalBox::Slot()
+		.FillHeight(0.90f)
+		.HAlign(HAlign_Fill)
+		[
+			SNew(SOverlay)
+
+			+ SOverlay::Slot()
+			[
+				GraphEditorPtr.ToSharedRef()
+			]
+
+			+SOverlay::Slot()
+			.VAlign(VAlign_Top)
+			.HAlign(HAlign_Fill)
+			.Padding(8)
+			[
+				SNew(SHorizontalBox)
+				+ SHorizontalBox::Slot()
+				.AutoWidth()
+				[
+					SNew(SBorder)
+					.BorderImage(FAppStyle::GetBrush("ToolPanel.GroupBorder"))
+					[
+						SNew(SVerticalBox)
+						+ SVerticalBox::Slot()
+						.AutoHeight()
+						[
+							SNew(SHorizontalBox)
+							+ SHorizontalBox::Slot()
+							.VAlign(VAlign_Center)
+							.Padding(2.f)
+							[
+								SNew(STextBlock)
+								.Text(LOCTEXT("SearchDepthReferencersLabelText", "Search Referencers Depth"))
+							]
+
+							+ SHorizontalBox::Slot()
+							.AutoWidth()
+							.VAlign(VAlign_Center)
+							.Padding(2.f)
+							[
+								SNew(SBox)
+								.WidthOverride(100)
+								[
+									SAssignNew(ReferencerCountBox, SSpinBox<int32>)
+									.Value(this, &SPluginReferenceViewer::GetSearchReferencerDepthCount)
+									.OnValueChanged_Lambda([this] (int32 NewValue)
+										{	
+											if (NewValue != Settings.MaxSearchReferencersDepth)
+											{
+												Settings.MaxSearchReferencersDepth = NewValue;
+												bNeedsGraphRebuild = true;
+											}
+										}
+									)
+									.OnValueCommitted_Lambda([this] (int32 NewValue, ETextCommit::Type CommitType) 
+										{ 
+											FSlateApplication::Get().SetKeyboardFocus(GraphEditorPtr, EFocusCause::SetDirectly); 
+
+											if (NewValue != Settings.MaxSearchReferencersDepth || bNeedsGraphRebuild)
+											{
+												Settings.MaxSearchReferencersDepth = NewValue;
+												bNeedsGraphRebuild = false;
+												RebuildGraph();
+											}
+										}
+									)
+									.MinValue(0)
+									.MaxValue(50)
+									.MaxSliderValue(10)
+								]
+							]
+						]
+
+						+ SVerticalBox::Slot()
+						.AutoHeight()
+						[
+							SNew(SHorizontalBox)
+							+ SHorizontalBox::Slot()
+							.VAlign(VAlign_Center)
+							.Padding(2.f)
+							[
+								SNew(STextBlock)
+								.Text(LOCTEXT("SearchDepthDependenciesLabelText", "Search Dependencies Depth"))
+							]
+
+							+ SHorizontalBox::Slot()
+							.AutoWidth()
+							.VAlign(VAlign_Center)
+							.Padding(2.f)
+							[
+								SNew(SBox)
+								.WidthOverride(100)
+								[
+									SAssignNew(DependencyCountBox, SSpinBox<int32>)
+									.Value(this, &SPluginReferenceViewer::GetSearchDependencyDepthCount)
+									.OnValueChanged_Lambda([this] (int32 NewValue)
+										{	
+											if (NewValue != Settings.MaxSearchDependencyDepth)
+											{
+												Settings.MaxSearchDependencyDepth = NewValue;
+												bNeedsGraphRebuild = true;
+											}
+										}
+									)
+									.OnValueCommitted_Lambda([this] (int32 NewValue, ETextCommit::Type CommitType) 
+										{ 
+											FSlateApplication::Get().SetKeyboardFocus(GraphEditorPtr, EFocusCause::SetDirectly); 
+
+											if (NewValue != Settings.MaxSearchDependencyDepth || bNeedsGraphRebuild)
+											{
+												Settings.MaxSearchDependencyDepth = NewValue;
+												bNeedsGraphRebuild = false;
+												RebuildGraph();
+											}
+										}
+									)
+									.MinValue(0)
+									.MaxValue(50)
+									.MaxSliderValue(10)
+								]
+							]
+						]
+					]
+				]
+			]
+		]
+	];
+}
+
+void SPluginReferenceViewer::SetGraphRoot(const TArray<FPluginIdentifier>& GraphRootIdentifiers)
+{
+	GraphObj->SetGraphRoot(GraphRootIdentifiers);
+
+	GraphObj->RebuildGraph();
+
+	GraphEditorPtr->ZoomToFit(false);
+}
+
+void SPluginReferenceViewer::OnOpenPluginProperties()
+{
+	TSharedPtr<const IPlugin> Plugin;
+
+	const FGraphPanelSelectionSet& SelectedNodes = GraphEditorPtr->GetSelectedNodes();
+	if (!SelectedNodes.IsEmpty())
+	{
+		for (UObject* Node : SelectedNodes)
+		{
+			UEdGraphNode_PluginReference* PluginReferenceNode = Cast<UEdGraphNode_PluginReference>(Node);
+			if (PluginReferenceNode)
+			{
+				Plugin = PluginReferenceNode->GetPlugin();
+				break;
+			}
+		}
+	}
+
+	if (Plugin != nullptr)
+	{
+		OpenPluginProperties(Plugin->GetName());
+	}
+}
+
+bool SPluginReferenceViewer::HasAtLeastOneRealNodeSelected()
+{
+	return true;
+}
+
+void SPluginReferenceViewer::OpenPluginProperties(const FString& PluginName)
+{
+	TSharedPtr<IPlugin> Plugin = IPluginManager::Get().FindPlugin(PluginName);
+	if (Plugin != nullptr)
+	{
+		IPluginsEditorFeature& PluginEditor = IModularFeatures::Get().GetModularFeature<IPluginsEditorFeature>(EditorFeatures::PluginsEditor);
+		PluginEditor.OpenPluginEditor(Plugin.ToSharedRef(), nullptr, FSimpleDelegate());
+	}
+}
+
+TSharedRef<SWidget> SPluginReferenceViewer::MakeToolBar()
+{
+	FToolBarBuilder ToolBarBuilder(PluginReferenceViewerActions, FMultiBoxCustomization::None, TSharedPtr<FExtender>(), true);
+
+	ToolBarBuilder.AddComboButton(
+		FUIAction(),
+		FOnGetContent::CreateSP(this, &SPluginReferenceViewer::GetShowMenuContent),
+		TAttribute<FText>(),
+		TAttribute<FText>(),
+		FSlateIcon(FAppStyle::GetAppStyleSetName(), "Icons.Visibility"),
+		/*bInSimpleComboBox*/ false);
+	
+	return ToolBarBuilder.MakeWidget();
+}
+
+TSharedRef<SWidget> SPluginReferenceViewer::GetShowMenuContent()
+{
+	FMenuBuilder MenuBuilder(true, PluginReferenceViewerActions);
+
+	MenuBuilder.BeginSection("ViewOptions", LOCTEXT("ViewOptions", "View Options"));
+	MenuBuilder.AddMenuEntry(FPluginReferenceViewerCommands::Get().CompactMode);
+	MenuBuilder.AddMenuEntry(FPluginReferenceViewerCommands::Get().ShowEnginePlugins);
+	MenuBuilder.AddMenuEntry(FPluginReferenceViewerCommands::Get().ShowOptionalPlugins);
+	MenuBuilder.EndSection();
+
+	return MenuBuilder.MakeWidget();
+}
+
+void SPluginReferenceViewer::RebuildGraph()
+{
+	GraphObj->RebuildGraph();
+}
+
+void SPluginReferenceViewer::ZoomToFit()
+{
+	GraphEditorPtr->ZoomToFit(true);
+}
+
+void SPluginReferenceViewer::ReCenterGraph()
+{
+	ReCenterGraphOnNodes(GraphEditorPtr->GetSelectedNodes());
+}
+
+void SPluginReferenceViewer::ReCenterGraphOnNodes(const TSet<UObject*>& Nodes)
+{
+	TArray<FPluginIdentifier> NewGraphRootNames;
+	FIntPoint TotalNodePos(ForceInitToZero);
+	for (auto NodeIt = Nodes.CreateConstIterator(); NodeIt; ++NodeIt)
+	{
+		UEdGraphNode_PluginReference* PluginReferenceNode = Cast<UEdGraphNode_PluginReference>(*NodeIt);
+		if (PluginReferenceNode)
+		{
+			NewGraphRootNames.Add(PluginReferenceNode->GetIdentifier());
+			TotalNodePos.X += PluginReferenceNode->NodePosX;
+			TotalNodePos.Y += PluginReferenceNode->NodePosY;
+		}
+	}
+
+	if (NewGraphRootNames.Num() > 0)
+	{
+		const FIntPoint AverageNodePos = TotalNodePos / NewGraphRootNames.Num();
+		GraphObj->SetGraphRoot(NewGraphRootNames, AverageNodePos);
+		UEdGraphNode_PluginReference* NewRootNode = GraphObj->RebuildGraph();
+
+		if (NewRootNode != nullptr)
+		{
+			GraphEditorPtr->ClearSelectionSet();
+			GraphEditorPtr->SetNodeSelection(NewRootNode, true);
+		}
+	}
+}
+
+void SPluginReferenceViewer::OnNodeDoubleClicked(UEdGraphNode* Node)
+{
+	if (UEdGraphNode_PluginReference* PlugingReferenceNode = Cast<UEdGraphNode_PluginReference>(Node))
+	{
+		TSet<UObject*> Nodes;
+		Nodes.Add(Node);
+		ReCenterGraphOnNodes(Nodes);
+	}
+}
+
+void SPluginReferenceViewer::RegisterActions()
+{
+	PluginReferenceViewerActions = MakeShareable(new FUICommandList);
+	FPluginReferenceViewerCommands::Register();
+
+	PluginReferenceViewerActions->MapAction(
+		FPluginReferenceViewerCommands::Get().CompactMode,
+		FExecuteAction::CreateSP(this, &SPluginReferenceViewer::OnCompactModeChanged),
+		FCanExecuteAction(),
+		FIsActionChecked::CreateSP(this, &SPluginReferenceViewer::IsCompactModeChecked));
+
+	PluginReferenceViewerActions->MapAction(
+		FPluginReferenceViewerCommands::Get().ShowEnginePlugins,
+		FExecuteAction::CreateSP(this, &SPluginReferenceViewer::OnShowEnginePluginsChanged),
+		FCanExecuteAction(),
+		FIsActionChecked::CreateSP(this, &SPluginReferenceViewer::IsShowEnginePluginsChecked));
+
+	PluginReferenceViewerActions->MapAction(
+		FPluginReferenceViewerCommands::Get().ShowOptionalPlugins,
+		FExecuteAction::CreateSP(this, &SPluginReferenceViewer::OnShowOptionalPluginsChanged),
+		FCanExecuteAction(),
+		FIsActionChecked::CreateSP(this, &SPluginReferenceViewer::IsShowOptionalPluginsChecked));
+
+	PluginReferenceViewerActions->MapAction(
+		FPluginReferenceViewerCommands::Get().OpenPluginProperties,
+		FExecuteAction::CreateSP(this, &SPluginReferenceViewer::OnOpenPluginProperties),
+		FCanExecuteAction::CreateSP(this, &SPluginReferenceViewer::HasAtLeastOneRealNodeSelected));
+
+	PluginReferenceViewerActions->MapAction(
+		FPluginReferenceViewerCommands::Get().ZoomToFit,
+		FExecuteAction::CreateSP(this, &SPluginReferenceViewer::ZoomToFit),
+		FCanExecuteAction());
+
+	PluginReferenceViewerActions->MapAction(
+		FPluginReferenceViewerCommands::Get().ReCenterGraph,
+		FExecuteAction::CreateSP(this, &SPluginReferenceViewer::ReCenterGraph),
+		FCanExecuteAction());
+}
+
+bool SPluginReferenceViewer::IsCompactModeChecked()
+{
+	return Settings.bIsCompactMode;
+}
+
+void SPluginReferenceViewer::OnCompactModeChanged()
+{
+	Settings.bIsCompactMode = !Settings.bIsCompactMode;
+	GraphObj->RebuildGraph();
+}
+
+bool SPluginReferenceViewer::IsShowEnginePluginsChecked() const
+{
+	return Settings.bShowEnginePlugins;
+}
+
+void SPluginReferenceViewer::OnShowEnginePluginsChanged()
+{
+	Settings.bShowEnginePlugins = !Settings.bShowEnginePlugins;
+	GraphObj->RebuildGraph();
+}
+
+bool SPluginReferenceViewer::IsShowOptionalPluginsChecked() const
+{
+	return Settings.bShowOptionalPlugins;
+}
+
+void SPluginReferenceViewer::OnShowOptionalPluginsChanged()
+{
+	Settings.bShowOptionalPlugins = !Settings.bShowOptionalPlugins;
+	GraphObj->RebuildGraph();
+}
+
+int32 SPluginReferenceViewer::GetSearchReferencerDepthCount() const
+{
+	return Settings.MaxSearchReferencersDepth;
+}
+
+int32 SPluginReferenceViewer::GetSearchDependencyDepthCount() const
+{
+	return Settings.MaxSearchDependencyDepth;
+}
+
+#undef LOCTEXT_NAMESPACE
\ No newline at end of file
diff --git a/Engine/Plugins/Experimental/PluginAudit/Source/Private/SPluginReferenceViewer.h b/Engine/Plugins/Experimental/PluginAudit/Source/Private/SPluginReferenceViewer.h
new file mode 100644
--- /dev/null
+++ b/Engine/Plugins/Experimental/PluginAudit/Source/Private/SPluginReferenceViewer.h
@@ -0,0 +1,81 @@
+// Copyright Epic Games, Inc. All Rights Reserved.
+
+#pragma once
+
+#include "Widgets/SCompoundWidget.h"
+#include "EdGraph_PluginReferenceViewer.h"
+#include "GraphEditor.h"
+
+//////////////////////////////////////////////////////////////////////////
+// SPluginReferenceViewer
+
+class IPlugin;
+class FUICommandInfo;
+class FUICommandList;
+class UEdGraphNode;
+class UEdGraph_PluginReferenceViewer;
+
+class SPluginReferenceViewer : public SCompoundWidget
+{
+public:
+	SLATE_BEGIN_ARGS(SPluginReferenceViewer)
+	{}
+
+	SLATE_END_ARGS()
+
+	~SPluginReferenceViewer();
+
+public:
+	void Construct(const FArguments& InArgs);
+
+	void SetGraphRoot(const TArray<FPluginIdentifier>& GraphRootIdentifiers);
+
+	int32 GetSearchReferencerDepthCount() const;
+	int32 GetSearchDependencyDepthCount() const;
+	bool IsCompactModeChecked();
+	bool IsShowEnginePluginsChecked() const;
+	bool IsShowOptionalPluginsChecked() const;
+
+private:
+	void OnOpenPluginProperties();
+	bool HasAtLeastOneRealNodeSelected();
+	void OpenPluginProperties(const FString& PluginName);
+
+	TSharedRef<SWidget> MakeToolBar();
+	TSharedRef<SWidget> GetShowMenuContent();
+
+	void RebuildGraph();
+	void ZoomToFit();
+	void ReCenterGraph();
+	void ReCenterGraphOnNodes(const TSet<UObject*>& Nodes);
+	void RegisterActions();
+
+	void OnNodeDoubleClicked(UEdGraphNode* Node);
+
+	void OnCompactModeChanged();
+	void OnShowEnginePluginsChanged();
+	void OnShowOptionalPluginsChanged();
+
+	struct FPluginReferenceViewerSettings
+	{
+		int32 MaxSearchReferencersDepth = 1;
+		int32 MaxSearchDependencyDepth = 1;
+		bool bIsCompactMode = false;
+		bool bShowEnginePlugins = true;
+		bool bShowOptionalPlugins = true;
+	};
+
+private:
+	TSharedPtr<SGraphEditor> GraphEditorPtr;
+	TSharedPtr<FUICommandList> PluginReferenceViewerActions;
+	TSharedPtr<SWidget> ReferencerCountBox;
+	TSharedPtr<SWidget> DependencyCountBox;
+
+	UEdGraph_PluginReferenceViewer* GraphObj;
+	FPluginReferenceViewerSettings Settings;
+
+	TSharedPtr<FUICommandInfo> ShowEnginePlugins;
+
+	/** Used to delay graph rebuilding during spinbox slider interaction */
+	bool bNeedsGraphRebuild;
+};
\ No newline at end of file
