Double clicking a plugin name entry in the audit list will open the plugin reference viewer focused on that plugin

#rnx

diff --git a/Engine/Plugins/Experimental/PluginAudit/Source/Private/SPluginAuditBrowser.cpp b/Engine/Plugins/Experimental/PluginAudit/Source/Private/SPluginAuditBrowser.cpp
--- a/Engine/Plugins/Experimental/PluginAudit/Source/Private/SPluginAuditBrowser.cpp
+++ b/Engine/Plugins/Experimental/PluginAudit/Source/Private/SPluginAuditBrowser.cpp
@@ -60,7 +60,8 @@ void SPluginAuditBrowser::Construct(const FArguments& InArgs)
 		SNew(SListView< TSharedRef<FCookedPlugin> >)
 		.ListItemsSource(&FilteredCookedPlugins)
 		.OnGenerateRow(this, &SPluginAuditBrowser::MakeCookedPluginRow)
-		.OnContextMenuOpening(this, &SPluginAuditBrowser::OnContextMenuOpening);
+		.OnContextMenuOpening(this, &SPluginAuditBrowser::OnContextMenuOpening)
+		.OnMouseButtonDoubleClick(this, &SPluginAuditBrowser::OnListViewDoubleClick);
 
 	ChildSlot
 	[
@@ -319,17 +320,22 @@ TSharedPtr<SWidget> SPluginAuditBrowser::OnContextMenuOpening()
 		));
 
 	MenuBuilder.AddMenuEntry(
-		LOCTEXT("OpenPluginDependencyViewerLabel", "Open Plugin Dependency Viewer"),
-		LOCTEXT("OpenPluginDependencyViewerTooltip", "Open Plugin Dependency Viewer"),
+		LOCTEXT("OpenPluginReferenceViewerLabel", "Open Plugin Reference Viewer"),
+		LOCTEXT("OpenPluginReferenceViewerTooltip", "Open Plugin Reference Viewer"),
 		FSlateIcon(),
 		FUIAction(
-			FExecuteAction::CreateSP(this, &SPluginAuditBrowser::OnOpenPluginDependencyViewer),
+			FExecuteAction::CreateSP(this, &SPluginAuditBrowser::OnOpenPluginReferenceViewer),
 			FCanExecuteAction()
 		));
 
 	return MenuBuilder.MakeWidget();
 }
 
+void SPluginAuditBrowser::OnListViewDoubleClick(TSharedRef<FCookedPlugin> Item)
+{
+	OpenPluginReferenceViewer(Item->Plugin);
+}
+
 void SPluginAuditBrowser::OnOpenPluginProperties()
 {
 	TArray<TSharedRef<FCookedPlugin>> SelectedItems = PluginListView->GetSelectedItems();
@@ -339,12 +345,12 @@ void SPluginAuditBrowser::OnOpenPluginProperties()
 	}
 }
 
-void SPluginAuditBrowser::OnOpenPluginDependencyViewer()
+void SPluginAuditBrowser::OnOpenPluginReferenceViewer()
 {
 	TArray<TSharedRef<FCookedPlugin>> SelectedItems = PluginListView->GetSelectedItems();
 	if (SelectedItems.Num() == 1)
 	{
-		FPluginReferenceViewerModule::Get().OpenPluginReferenceViewerUI(SelectedItems[0]->Plugin);
+		OpenPluginReferenceViewer(SelectedItems[0]->Plugin);
 	}
 }
 
@@ -358,6 +364,11 @@ void SPluginAuditBrowser::OpenPluginProperties(const FString& PluginName)
 	}
 }
 
+void SPluginAuditBrowser::OpenPluginReferenceViewer(const TSharedRef<IPlugin>& Plugin)
+{
+	FPluginReferenceViewerModule::Get().OpenPluginReferenceViewerUI(Plugin);
+}
+
 void SPluginAuditBrowser::CreateLogListing()
 {
 	FMessageLogModule& MessageLogModule = FModuleManager::LoadModuleChecked<FMessageLogModule>("MessageLog");
diff --git a/Engine/Plugins/Experimental/PluginAudit/Source/Private/SPluginAuditBrowser.h b/Engine/Plugins/Experimental/PluginAudit/Source/Private/SPluginAuditBrowser.h
--- a/Engine/Plugins/Experimental/PluginAudit/Source/Private/SPluginAuditBrowser.h
+++ b/Engine/Plugins/Experimental/PluginAudit/Source/Private/SPluginAuditBrowser.h
@@ -19,7 +19,6 @@ struct FGameplayTag;
 class IMessageLogListing;
 class IMessageToken;
 class SFilterSearchBox;
-class SPluginReferenceViewer;
 class STableViewBase;
 class ITableRow;
 class FTokenizedMessage;
@@ -48,8 +47,9 @@ private:
 	FText GetPluginCountText() const;
 	TSharedPtr<SWidget> OnContextMenuOpening();
 	void OnOpenPluginProperties();
-	void OnOpenPluginDependencyViewer();
+	void OnOpenPluginReferenceViewer();
 	void OpenPluginProperties(const FString& PluginName);
+	void OpenPluginReferenceViewer(const TSharedRef<IPlugin>& Plugin);
 
 	class FCookedPlugin
 	{
@@ -63,6 +63,8 @@ private:
 		bool bSimulateDisabled = false;
 	};
 
+	void OnListViewDoubleClick(TSharedRef<FCookedPlugin> Item);
+
 	TSharedRef<ITableRow> MakeCookedPluginRow(TSharedRef<FCookedPlugin> InItem, const TSharedRef<STableViewBase>& OwnerTable);
 
 	struct FGameFeaturePlugin
@@ -105,8 +107,6 @@ private:
 	/** The plugin search box widget */
 	TSharedPtr<class SSearchBox> SearchBoxPtr;
 
-	TSharedPtr<SPluginReferenceViewer> PluginReferenceViewer;
-
 	/** Text filter object for typing in filter text to the search box */
 	TSharedPtr<FCookedPluginTextFilter> PluginTextFilter;
 
