Added ability to cancel the Plugin Audit process in the editor.

diff --git a/Engine/Plugins/Experimental/PluginAudit/Source/Private/SPluginAuditBrowser.cpp b/Engine/Plugins/Experimental/PluginAudit/Source/Private/SPluginAuditBrowser.cpp
--- a/Engine/Plugins/Experimental/PluginAudit/Source/Private/SPluginAuditBrowser.cpp
+++ b/Engine/Plugins/Experimental/PluginAudit/Source/Private/SPluginAuditBrowser.cpp
@@ -423,14 +423,22 @@ void SPluginAuditBrowser::RefreshViolations()
 	FCriticalSection CS;
 
 	FScopedSlowTask SlowTask(IncludedPlugins.Num() + ExcludedPlugins.Num(), LOCTEXT("Examining Plugins", "Examining Plugins..."));
-	SlowTask.MakeDialog();
+	SlowTask.MakeDialog(true);
 
 	const FName GameplayTagStructPackage = FGameplayTag::StaticStruct()->GetOutermost()->GetFName();
 	const FName NAME_GameplayTag = FGameplayTag::StaticStruct()->GetFName();
+
+	bool bCancelled = false;
 	
 	// Included plugins.
 	for (const FGameFeaturePlugin& Plugin : IncludedPlugins)
 	{
+		if (SlowTask.ShouldCancel())
+		{
+			bCancelled = true;
+			break;
+		}
+
 		SlowTask.EnterProgressFrame(1.f, FText::FromString(Plugin.Plugin->GetFriendlyName()));
 
 		FARFilter Filter;
@@ -518,9 +526,20 @@ void SPluginAuditBrowser::RefreshViolations()
 		});
 	}
 
+	if (bCancelled)
+	{
+		return Violations;
+	}
+
 	// Excluded plugins
 	for (const FGameFeaturePlugin& Plugin : ExcludedPlugins)
 	{
+		if (SlowTask.ShouldCancel())
+		{
+			bCancelled = true;
+			break;
+		}
+
 		SlowTask.EnterProgressFrame(1.f, FText::FromString(Plugin.Plugin->GetFriendlyName()));
 		
 		FARFilter Filter;
